unsync - Rename weak hash ban list variable

diff --git a/Engine/Source/Programs/Unsync/Private/UnsyncCore.cpp b/Engine/Source/Programs/Unsync/Private/UnsyncCore.cpp
--- a/Engine/Source/Programs/Unsync/Private/UnsyncCore.cpp
+++ b/Engine/Source/Programs/Unsync/Private/UnsyncCore.cpp
@@ -911,10 +911,10 @@ DiffBlocksParallelT(FIOReader&				  BaseDataReader,
 
 					const uint32							 MaxWeakHashFalsePositives = 8;
 					HashMap<uint32, uint32, FIdentityHash32> WeakHashFalsePositives;
-					HashSet<uint32, FIdentityHash32>		 WeakHashBlackList;
+					HashSet<uint32, FIdentityHash32>		 WeakHashBanList;
 
 					auto ScanFn = [&SourceWeakHashSet,
-								   &WeakHashBlackList,
+								   &WeakHashBanList,
 								   BlockSize,
 								   &Task,
 								   TaskBuffer,
@@ -926,7 +926,7 @@ DiffBlocksParallelT(FIOReader&				  BaseDataReader,
 						uint64 ThisBlockSize = WindowEnd - WindowBegin;
 
 						if (SourceWeakHashSet.find(WindowHash) != SourceWeakHashSet.end() &&
-							WeakHashBlackList.find(WindowHash) == WeakHashBlackList.end())
+							WeakHashBanList.find(WindowHash) == WeakHashBanList.end())
 						{
 							UNSYNC_ASSERT(ThisBlockSize <= BlockSize);
 
@@ -950,7 +950,7 @@ DiffBlocksParallelT(FIOReader&				  BaseDataReader,
 							uint32 FalsePositives = WeakHashFalsePositives[WindowHash]++;
 							if (FalsePositives >= MaxWeakHashFalsePositives)
 							{
-								WeakHashBlackList.insert(WindowHash);
+								WeakHashBanList.insert(WindowHash);
 							}
 						}
 
