unsync - Remove VS2022 requirement

diff --git a/Engine/Source/Programs/Unsync/Unsync.Target.cs b/Engine/Source/Programs/Unsync/Unsync.Target.cs
--- a/Engine/Source/Programs/Unsync/Unsync.Target.cs
+++ b/Engine/Source/Programs/Unsync/Unsync.Target.cs
@@ -19,7 +19,6 @@ public class UnsyncTarget : TargetRules
 		bEnforceIWYU = false;
 		GlobalDefinitions.Add("_CRT_SECURE_NO_WARNINGS");
 		WindowsPlatform.TargetWindowsVersion = 0x0603; // Windows 8.1
-		WindowsPlatform.Compiler = WindowsCompiler.VisualStudio2022;
 		bUseStaticCRT = true;
 		bIsBuildingConsoleApplication = true;
 	}
