unsync - Fix compile errors on Mac

diff --git a/Engine/Source/Programs/Unsync/Private/UnsyncCommon.h b/Engine/Source/Programs/Unsync/Private/UnsyncCommon.h
--- a/Engine/Source/Programs/Unsync/Private/UnsyncCommon.h
+++ b/Engine/Source/Programs/Unsync/Private/UnsyncCommon.h
@@ -67,6 +67,7 @@ ClobberT(T& Object)
 #	define UNSYNC_ATTRIB_FORCEINLINE __attribute__((always_inline))
 #	define UNSYNC_THIRD_PARTY_INCLUDES_START \
 		_Pragma("GCC diagnostic push") \
+		_Pragma("GCC diagnostic ignored \"-Wunknown-pragmas\"") \
 		_Pragma("GCC diagnostic ignored \"-Wshadow\"")
 #	define UNSYNC_THIRD_PARTY_INCLUDES_END _Pragma("GCC diagnostic pop")
 #endif
diff --git a/Engine/Source/Programs/Unsync/Private/UnsyncCore.cpp b/Engine/Source/Programs/Unsync/Private/UnsyncCore.cpp
--- a/Engine/Source/Programs/Unsync/Private/UnsyncCore.cpp
+++ b/Engine/Source/Programs/Unsync/Private/UnsyncCore.cpp
@@ -2217,7 +2217,6 @@ ValidateTarget(FIOReader& Reader, const FNeedList& NeedList, EStrongHashAlgorith
 	FSemaphore IoSemaphore(MAX_ACTIVE_READERS);
 
 	const uint64		TotalStreamBytes = Reader.GetSize();
-	std::atomic<uint64> ReadBytesTotal	 = {};
 	std::atomic<uint64> NumInvalidBlocks = {};
 	FTaskGroup			TaskGroup;
 
@@ -2263,7 +2262,6 @@ ValidateTarget(FIOReader& Reader, const FNeedList& NeedList, EStrongHashAlgorith
 					  (ReadOffset + BatchSizeBytes) == ValidationBlocks[BlockIndex + 1].Offset);
 
 		auto ReadCallback = [StrongHasher,
-							 TotalStreamBytes,
 							 bLogVerbose,
 							 LogIndent,
 							 BatchBegin,
@@ -2285,7 +2283,6 @@ ValidateTarget(FIOReader& Reader, const FNeedList& NeedList, EStrongHashAlgorith
 						   BatchBegin,
 						   BatchEnd,
 						   StrongHasher,
-						   TotalStreamBytes,
 						   bLogVerbose,
 						   LogIndent,
 						   &NumInvalidBlocks,
@@ -3250,7 +3247,7 @@ SyncDirectory(const FSyncDirectoryOptions& SyncOptions)
 				RemainingSourceBytes -= LocalTask.NeedBytesFromSource;
 
 				ForegroundTaskGroup.run(
-					[Task = std::move(LocalTask), &ProxyPool, &NumForegroundTasks, &SyncTaskBody, LogVerbose = GLogVerbose]() {
+					[Task = std::move(LocalTask), &NumForegroundTasks, &SyncTaskBody, LogVerbose = GLogVerbose]() {
 						FLogVerbosityScope VerbosityScope(LogVerbose);
 						SyncTaskBody(Task, false);
 						--NumForegroundTasks;
@@ -3274,7 +3271,7 @@ SyncDirectory(const FSyncDirectoryOptions& SyncOptions)
 				RemainingSourceBytes -= LocalTask.NeedBytesFromSource;
 
 				BackgroundTaskGroup.run(
-					[Task = std::move(LocalTask), &NumBackgroundTasks, &SyncTaskBody, &CurrentBackgroundMemory, &ProxyPool]() {
+					[Task = std::move(LocalTask), &NumBackgroundTasks, &SyncTaskBody, &CurrentBackgroundMemory]() {
 						FLogVerbosityScope VerbosityScope(false);  // turn off logging from background threads
 						SyncTaskBody(Task, true);
 						--NumBackgroundTasks;
diff --git a/Engine/Source/Programs/Unsync/Private/UnsyncFile.h b/Engine/Source/Programs/Unsync/Private/UnsyncFile.h
--- a/Engine/Source/Programs/Unsync/Private/UnsyncFile.h
+++ b/Engine/Source/Programs/Unsync/Private/UnsyncFile.h
@@ -297,8 +297,10 @@ struct FMemReader : FIOReader
 	uint64		 Size = 0;
 };
 
+#if UNSYNC_COMPILER_MSVC
 #pragma warning(push)
 #pragma warning(disable : 4250)	 // 'FMemReaderWriter': inherits 'FMemReader::FMemReader::flush_one' via dominance
+#endif // UNSYNC_COMPILER_MSVC
 struct FMemReaderWriter : FMemReader, FIOReaderWriter
 {
 	FMemReaderWriter(uint8* InData, uint64 InDataSize);
@@ -321,7 +323,9 @@ struct FMemReaderWriter : FMemReader, FIOReaderWriter
 
 	uint8* DataRw = nullptr;
 };
+#if UNSYNC_COMPILER_MSVC
 #pragma warning(pop)
+#endif // UNSYNC_COMPILER_MSVC
 
 struct FNullReaderWriter : FIOReaderWriter
 {
@@ -374,7 +378,7 @@ struct FDeferredOpenReader : FIOReader
 	{
 		return GetOrOpenInner()->Read(Dest, SourceOffset, ReadSize);
 	}
-	virtual bool ReadAsync(uint64 SourceOffset, uint64 Size, uint64 UserData, IOCallback Callback)
+	virtual bool ReadAsync(uint64 SourceOffset, uint64 Size, uint64 UserData, IOCallback Callback) override
 	{
 		return GetOrOpenInner()->ReadAsync(SourceOffset, Size, UserData, Callback);
 	}
diff --git a/Engine/Source/Programs/Unsync/Private/UnsyncHash.cpp b/Engine/Source/Programs/Unsync/Private/UnsyncHash.cpp
--- a/Engine/Source/Programs/Unsync/Private/UnsyncHash.cpp
+++ b/Engine/Source/Programs/Unsync/Private/UnsyncHash.cpp
@@ -4,10 +4,10 @@
 #include "UnsyncBuffer.h"
 #include "UnsyncProtocol.h"
 
-#pragma warning(push, 3)
+UNSYNC_THIRD_PARTY_INCLUDES_START
 #include <blake3.h>
 #include <md5-sse2.h>
-#pragma warning(pop)
+UNSYNC_THIRD_PARTY_INCLUDES_END
 
 namespace unsync {
 
diff --git a/Engine/Source/Programs/Unsync/Private/UnsyncSocket.h b/Engine/Source/Programs/Unsync/Private/UnsyncSocket.h
--- a/Engine/Source/Programs/Unsync/Private/UnsyncSocket.h
+++ b/Engine/Source/Programs/Unsync/Private/UnsyncSocket.h
@@ -178,7 +178,7 @@ RecvStruct(FSocketBase& Socket, const T& Data)
 {
 	int32 Size = int32(sizeof(Data));
 	bool  bOk  = true;
-	bOk &= SocketRecvT(Socket, Data);
+	bOk &= SocketRecvT(Socket, Size);
 	bOk &= SocketRecvT(Socket, Data);
 	return bOk;
 }
diff --git a/Engine/Source/Programs/Unsync/Private/UnsyncTest.cpp b/Engine/Source/Programs/Unsync/Private/UnsyncTest.cpp
--- a/Engine/Source/Programs/Unsync/Private/UnsyncTest.cpp
+++ b/Engine/Source/Programs/Unsync/Private/UnsyncTest.cpp
@@ -219,7 +219,7 @@ TestRollingSum()
 				const Task& Task	 = Tasks[TaskIdx];
 				uint64		TaskSize = Task.End - Task.Begin;
 
-				auto ScanFn = [&HashReference, &TaskSize, &TaskIdx, &HashRolling, &Task, DataBegin, DataEnd, &ExpectedIdx](
+				auto ScanFn = [&HashReference, &HashRolling, &Task, DataBegin, DataEnd, &ExpectedIdx](
 								  const uint8* WindowBegin,
 								  const uint8* WindowEnd,
 								  uint32	   WindowHash) {
diff --git a/Engine/Source/Programs/Unsync/Private/UnsyncThread.h b/Engine/Source/Programs/Unsync/Private/UnsyncThread.h
--- a/Engine/Source/Programs/Unsync/Private/UnsyncThread.h
+++ b/Engine/Source/Programs/Unsync/Private/UnsyncThread.h
@@ -16,7 +16,19 @@ UNSYNC_THIRD_PARTY_INCLUDES_START
 #else
 #	include <semaphore>
 #	include <thread>
+#	ifdef __APPLE__
+#		define UNSYNC_USE_MACH_SEMAPHORE 1
+#	endif
 #endif	// UNSYNC_USE_CONCRT
+#if !defined(UNSYNC_USE_MACH_SEMAPHORE)
+#	define UNSYNC_USE_MACH_SEMAPHORE 0
+#endif
+#if UNSYNC_USE_MACH_SEMAPHORE
+#include <mach/mach_init.h>
+#include <mach/mach_error.h>
+#include <mach/semaphore.h>
+#include <mach/task.h>
+#endif // UNSYNC_USE_MACH_SEMAPHORE
 UNSYNC_THIRD_PARTY_INCLUDES_END
 
 namespace unsync {
@@ -110,6 +122,42 @@ ParallelForEach(IT ItBegin, IT ItEnd, FT F)
 
 #else  // UNSYNC_USE_CONCRT
 
+#if UNSYNC_USE_MACH_SEMAPHORE
+
+class FSemaphore
+{
+public:
+	UNSYNC_DISALLOW_COPY_ASSIGN(FSemaphore)
+
+	explicit FSemaphore(uint32 MaxCount)
+	{
+		kern_return_t InitResult = semaphore_create(mach_task_self(), &Native, SYNC_POLICY_FIFO, MaxCount);
+		UNSYNC_ASSERTF(InitResult == KERN_SUCCESS, L"Failed to create a semaphore, error code: %d %hs", InitResult, mach_error_string(InitResult));
+	}
+
+	~FSemaphore()
+	{
+		kern_return_t DestroyResult = semaphore_destroy(mach_task_self(), Native);
+		UNSYNC_ASSERTF(DestroyResult == KERN_SUCCESS, L"Failed to destroy a semaphore, error code: %d %hs", DestroyResult, mach_error_string(DestroyResult));
+	}
+
+	void Acquire()
+	{
+		kern_return_t WaitResult = semaphore_wait(Native);
+		UNSYNC_ASSERTF(WaitResult == KERN_SUCCESS, L"Failed to wait for a semaphore, error code: %d %hs", WaitResult, mach_error_string(WaitResult));
+	}
+
+	void Release()
+	{
+		int32 SignalResult = semaphore_signal(Native);
+		UNSYNC_ASSERTF(SignalResult == 0, L"Failed to signal a semaphore, error code: %d %hs", SignalResult, mach_error_string(SignalResult));
+	}
+
+	semaphore_t Native = {};
+};
+
+#else // UNSYNC_USE_MACH_SEMAPHORE
+
 class FSemaphore
 {
 public:
@@ -124,9 +172,12 @@ public:
 	void Release() { Native.release(); }
 
 private:
+
 	std::counting_semaphore<UNSYNC_MAX_TOTAL_THREADS> Native;
 };
 
+#endif // UNSYNC_USE_MACH_SEMAPHORE
+
 // Single-threaded task group implementation
 struct FTaskGroup
 {
