unsync - Fix few minor build issues

* Suppress warnings in CLI11 headers
* Use WIN32_LEAN_AND_MEAN=1 and include specific headers when they're needed
Some windows headers (at least in some versions of the SDKs) do not pop the warning suppression states, which causes inconsistent warning reporting whenever Windows.h is included.
By adding WIN32_LEAN_AND_MEAN, we avoid the issue at least in the common case.
* Suppress various deprecation warnings in Windows headers (_SILENCE_CXX17_OLD_ALLOCATOR_MEMBERS_DEPRECATION_WARNING, _WINSOCK_DEPRECATED_NO_WARNINGS, _SILENCE_CXX17_CODECVT_HEADER_DEPRECATION_WARNING, _CRT_SECURE_NO_WARNINGS)
* Fix blake3 include path when using cmake

This fixes compilation warnings when using Windows SDK 10.0.22000.0.

diff --git a/Engine/Source/Programs/Unsync/CMakeLists.txt b/Engine/Source/Programs/Unsync/CMakeLists.txt
--- a/Engine/Source/Programs/Unsync/CMakeLists.txt
+++ b/Engine/Source/Programs/Unsync/CMakeLists.txt
@@ -116,7 +116,7 @@ find_package(fmt CONFIG REQUIRED)
 find_package(unofficial-http-parser CONFIG REQUIRED)
 find_package(zstd CONFIG REQUIRED)
 
-target_include_directories(unsync PRIVATE ${CMAKE_CURRENT_SOURCE_DIR}/thirdparty)
+target_include_directories(unsync PRIVATE ${CMAKE_CURRENT_SOURCE_DIR}/thirdparty ${blake3dir})
 target_compile_definitions(unsync PRIVATE BLAKE3_NO_SSE2=1 BLAKE3_NO_AVX512=1)
 
 target_link_directories(unsync PRIVATE ${VCPKG_INSTALLED_DIR}/${VCPKG_TARGET_TRIPLET}/lib)
@@ -141,10 +141,12 @@ if (UNSYNC_USE_DEBUG_HEAP)
 endif(UNSYNC_USE_DEBUG_HEAP)
 
 if(${CMAKE_SYSTEM_NAME} MATCHES "Windows")
-	target_compile_definitions(unsync PRIVATE 
+	target_compile_definitions(unsync PRIVATE
 		UNSYNC_PLATFORM_WINDOWS=1
 		UNSYNC_PLATFORM_UNIX=0
 		UNSYNC_USE_CONCRT=1
+		WIN32_LEAN_AND_MEAN=1
+		NOMINMAX=1
 	)
 else()
 	target_compile_definitions(unsync PRIVATE 
@@ -157,7 +159,6 @@ endif()
 if(MSVC)
 
 	add_definitions(-MP)
-	add_definitions(-DNOMINMAX=1)
 	add_definitions(-D_SILENCE_CXX17_OLD_ALLOCATOR_MEMBERS_DEPRECATION_WARNING=1)
 	add_definitions(-D_CRT_SECURE_NO_WARNINGS=1)
 	add_definitions(-D_WINSOCK_DEPRECATED_NO_WARNINGS=1)
diff --git a/Engine/Source/Programs/Unsync/Private/UnsyncFile.cpp b/Engine/Source/Programs/Unsync/Private/UnsyncFile.cpp
--- a/Engine/Source/Programs/Unsync/Private/UnsyncFile.cpp
+++ b/Engine/Source/Programs/Unsync/Private/UnsyncFile.cpp
@@ -14,6 +14,12 @@
 #	include <unistd.h>
 #endif	// UNSYNC_PLATFORM_UNIX
 
+#if UNSYNC_PLATFORM_WINDOWS
+UNSYNC_THIRD_PARTY_INCLUDES_START
+#include <winioctl.h>
+UNSYNC_THIRD_PARTY_INCLUDES_END
+#endif // UNSYNC_PLATFORM_WINDOWS
+
 namespace unsync {
 
 bool GForceBufferedFiles = false;
diff --git a/Engine/Source/Programs/Unsync/Private/UnsyncMain.cpp b/Engine/Source/Programs/Unsync/Private/UnsyncMain.cpp
--- a/Engine/Source/Programs/Unsync/Private/UnsyncMain.cpp
+++ b/Engine/Source/Programs/Unsync/Private/UnsyncMain.cpp
@@ -13,17 +13,18 @@
 #include "UnsyncThread.h"
 #include "UnsyncUtil.h"
 
+UNSYNC_THIRD_PARTY_INCLUDES_START
 #if UNSYNC_PLATFORM_WINDOWS
 #	include <io.h>
+#	include <shellapi.h>
 #endif	// UNSYNC_PLATFORM_WINDOWS
-
 #include <fcntl.h>
 #include <filesystem>
 #include <fstream>
 #include <iostream>
 #include <system_error>
-
 #include <CLI/CLI.hpp>
+UNSYNC_THIRD_PARTY_INCLUDES_END
 
 namespace unsync {
 
diff --git a/Engine/Source/Programs/Unsync/Unsync.Build.cs b/Engine/Source/Programs/Unsync/Unsync.Build.cs
--- a/Engine/Source/Programs/Unsync/Unsync.Build.cs
+++ b/Engine/Source/Programs/Unsync/Unsync.Build.cs
@@ -23,6 +23,11 @@ public class Unsync : ModuleRules
 			PrivateDefinitions.Add("UNSYNC_USE_CONCRT=1");
 			PrivateDefinitions.Add("UNSYNC_PLATFORM_UNIX=0");
 			PrivateDefinitions.Add("NOMINMAX=1");
+			PrivateDefinitions.Add("WIN32_LEAN_AND_MEAN=1");
+			PrivateDefinitions.Add("_SILENCE_CXX17_OLD_ALLOCATOR_MEMBERS_DEPRECATION_WARNING=1");
+			PrivateDefinitions.Add("_CRT_SECURE_NO_WARNINGS=1");
+			PrivateDefinitions.Add("_WINSOCK_DEPRECATED_NO_WARNINGS=1");
+			PrivateDefinitions.Add("_SILENCE_CXX17_CODECVT_HEADER_DEPRECATION_WARNING=1");
 		}
 		else if (Target.Platform == UnrealTargetPlatform.Mac || Target.Platform == UnrealTargetPlatform.Linux)
 		{
