unsync - Fix few badly auto-converted function names

Also fix missing header include

diff --git a/Engine/Source/Programs/Unsync/Private/UnsyncCmdHash.cpp b/Engine/Source/Programs/Unsync/Private/UnsyncCmdHash.cpp
--- a/Engine/Source/Programs/Unsync/Private/UnsyncCmdHash.cpp
+++ b/Engine/Source/Programs/Unsync/Private/UnsyncCmdHash.cpp
@@ -10,7 +10,7 @@ namespace unsync {
 int32  // TODO: return a TResult
 CmdHash(const FCmdHashOptions& Options)
 {
-	if (unsync::bDirectory(Options.Input))
+	if (unsync::IsDirectory(Options.Input))
 	{
 		UNSYNC_VERBOSE(L"'%ls' is a directory", Options.Input.wstring().c_str());
 
diff --git a/Engine/Source/Programs/Unsync/Private/UnsyncCmdQuery.cpp b/Engine/Source/Programs/Unsync/Private/UnsyncCmdQuery.cpp
--- a/Engine/Source/Programs/Unsync/Private/UnsyncCmdQuery.cpp
+++ b/Engine/Source/Programs/Unsync/Private/UnsyncCmdQuery.cpp
@@ -7,6 +7,7 @@
 
 #include <algorithm>
 #include <json11.hpp>
+#include <float.h>
 
 namespace unsync {
 
diff --git a/Engine/Source/Programs/Unsync/Private/UnsyncCmdSync.cpp b/Engine/Source/Programs/Unsync/Private/UnsyncCmdSync.cpp
--- a/Engine/Source/Programs/Unsync/Private/UnsyncCmdSync.cpp
+++ b/Engine/Source/Programs/Unsync/Private/UnsyncCmdSync.cpp
@@ -17,7 +17,7 @@ CmdSync(const FCmdSyncOptions& Options)
 	UNSYNC_VERBOSE(L"Sync target: '%ls'", Options.Target.wstring().c_str());
 
 	const bool bSourcePathExists	 = fs::exists(ResolvedSource, ErrorCode);
-	const bool bSourceIsDirectory	 = bSourcePathExists && unsync::bDirectory(ResolvedSource);
+	const bool bSourceIsDirectory	 = bSourcePathExists && unsync::IsDirectory(ResolvedSource);
 	const bool bSourceIsManifestHash = !bSourcePathExists && LooksLikeHash160(Options.Source.native());
 
 	bool bSourceFileSystemRequired = !bSourceIsManifestHash;
diff --git a/Engine/Source/Programs/Unsync/Private/UnsyncFile.cpp b/Engine/Source/Programs/Unsync/Private/UnsyncFile.cpp
--- a/Engine/Source/Programs/Unsync/Private/UnsyncFile.cpp
+++ b/Engine/Source/Programs/Unsync/Private/UnsyncFile.cpp
@@ -74,7 +74,7 @@ FWindowsFile::FWindowsFile(const fs::path& InFilename, EFileMode InMode, uint64
 
 	if (bOpenedOk)
 	{
-		if (bReadOnly(InMode))
+		if (IsReadOnly(InMode))
 		{
 			LARGE_INTEGER LiFileSize = {};
 			bool		  bSizeOk	 = GetFileSizeEx(FileHandle, &LiFileSize);
@@ -602,7 +602,7 @@ ToWindowsFileTime(const fs::file_time_type& T)
 #if UNSYNC_PLATFORM_UNIX
 FUnixFile::FUnixFile(const fs::path& InFilename, EFileMode InMode, uint64 in_size) : Filename(InFilename), Mode(InMode)
 {
-	FileHandle = fopen(InFilename.native().c_str(), bReadOnly(Mode) ? "rb" : "w+b");
+	FileHandle = fopen(InFilename.native().c_str(), IsReadOnly(Mode) ? "rb" : "w+b");
 	if (FileHandle == nullptr)
 	{
 		return;
@@ -610,7 +610,7 @@ FUnixFile::FUnixFile(const fs::path& InFilename, EFileMode InMode, uint64 in_siz
 
 	FileDescriptor = fileno(FileHandle);
 
-	if (bReadOnly(Mode))
+	if (IsReadOnly(Mode))
 	{
 		struct stat stat_buf = {};
 		LastError			 = fstat(FileDescriptor, &stat_buf);
@@ -1003,7 +1003,7 @@ CreateFileAttributeCache(const fs::path& Root, const FSyncFilter* SyncFilter)
 }
 
 bool
-bDirectory(const fs::path& Path)
+IsDirectory(const fs::path& Path)
 {
 	FileAttributes Attr = GetFileAttrib(Path);
 	return Attr.bValid && Attr.bDirectory;
diff --git a/Engine/Source/Programs/Unsync/Private/UnsyncFile.h b/Engine/Source/Programs/Unsync/Private/UnsyncFile.h
--- a/Engine/Source/Programs/Unsync/Private/UnsyncFile.h
+++ b/Engine/Source/Programs/Unsync/Private/UnsyncFile.h
@@ -23,7 +23,7 @@ enum class EFileMode
 };
 
 inline bool
-bReadOnly(EFileMode Mode)
+IsReadOnly(EFileMode Mode)
 {
 	switch (Mode)
 	{
@@ -444,7 +444,7 @@ struct FFileAttributeCache
 FileAttributes GetFileAttrib(const fs::path& Path, FFileAttributeCache* AttribCache = nullptr);
 bool		   SetFileMtime(const fs::path& Path, uint64 Mtime);
 bool		   SetFileReadOnly(const fs::path& Path, bool ReadOnly);
-bool		   bDirectory(const fs::path& Path);
+bool		   IsDirectory(const fs::path& Path);
 
 uint64 ToWindowsFileTime(const fs::file_time_type& T);
 
diff --git a/Engine/Source/Programs/Unsync/Private/UnsyncTest.cpp b/Engine/Source/Programs/Unsync/Private/UnsyncTest.cpp
--- a/Engine/Source/Programs/Unsync/Private/UnsyncTest.cpp
+++ b/Engine/Source/Programs/Unsync/Private/UnsyncTest.cpp
@@ -37,7 +37,7 @@ VerifyRandomBytes(const uint8* Data, uint64 Size, uint32 Seed)
 static bool
 TestRollingSum()
 {
-	UNSYNC_LOG(L"test_rolling_sum()");
+	UNSYNC_LOG(L"TestRollingSum()");
 
 	// sanity check
 	{
@@ -272,7 +272,7 @@ TestSync(EWeakHashAlgorithmID WeakHasher, EStrongHashAlgorithmID StrongHasher)
 	Algorithm.WeakHashAlgorithmId	= WeakHasher;
 	Algorithm.StrongHashAlgorithmId = StrongHasher;
 
-	UNSYNC_LOG(L"test_sync(%hs, %hs)", ToString(WeakHasher), ToString(StrongHasher));
+	UNSYNC_LOG(L"TestSync(%hs, %hs)", ToString(WeakHasher), ToString(StrongHasher));
 
 	{
 		uint32	BlockSize	 = 4;
@@ -435,7 +435,7 @@ TestBuildTarget(EWeakHashAlgorithmID WeakHasher, EStrongHashAlgorithmID StrongHa
 	Algorithm.WeakHashAlgorithmId	= WeakHasher;
 	Algorithm.StrongHashAlgorithmId = StrongHasher;
 
-	UNSYNC_LOG(L"test_BuildTarget(%hs, %hs)", ToString(WeakHasher), ToString(StrongHasher));
+	UNSYNC_LOG(L"TestBuildTarget(%hs, %hs)", ToString(WeakHasher), ToString(StrongHasher));
 	uint32	BlockSize = 4;
 	FBuffer Source	  = {0, 1, 2, 3, 4, 5, 6, 7, 8, 9, 10, 11, 99, 99, 99, 99};
 	FBuffer Base	  = {16, 17, 18, 19, 20, 21, 99, 99, 99, 99, 26, 27, 28, 29, 30, 31};
@@ -454,7 +454,7 @@ TestBuildTarget(EWeakHashAlgorithmID WeakHasher, EStrongHashAlgorithmID StrongHa
 void
 TestFiles()
 {
-	UNSYNC_LOG(L"test_files()");
+	UNSYNC_LOG(L"TestFiles()");
 
 	// TODO: test different chunk modes
 	const EChunkingAlgorithmID ChunkingMode = EChunkingAlgorithmID::FixedBlocks;
@@ -539,7 +539,7 @@ TestFiles()
 void
 TestPerfHashWeak(EWeakHashAlgorithmID Algorithm)
 {
-	UNSYNC_LOG(L"test_perf_hash_weak(%hs)", ToString(Algorithm));
+	UNSYNC_LOG(L"TestPerfHashWeak(%hs)", ToString(Algorithm));
 	UNSYNC_LOG_INDENT;
 
 	UNSYNC_LOG(L"Generating data");
@@ -578,7 +578,7 @@ TestPerfHashWeak(EWeakHashAlgorithmID Algorithm)
 void
 TestPerfHashStrong(EStrongHashAlgorithmID StrongHasher)
 {
-	UNSYNC_LOG(L"test_perf_hash_strong(%hs)", ToString(StrongHasher));
+	UNSYNC_LOG(L"TestPerfHashStrong(%hs)", ToString(StrongHasher));
 	UNSYNC_LOG_INDENT;
 
 	UNSYNC_LOG(L"Generating data");
@@ -606,7 +606,7 @@ TestPerfHashStrong(EStrongHashAlgorithmID StrongHasher)
 void
 TestPerfComputeBlocksVariable(EWeakHashAlgorithmID WeakHasher, EStrongHashAlgorithmID StrongHasher)
 {
-	UNSYNC_LOG(L"test_perf_ComputeBlocksVariable(%hs, %hs)", ToString(WeakHasher), ToString(StrongHasher));
+	UNSYNC_LOG(L"TestPerfComputeBlocksVariable(%hs, %hs)", ToString(WeakHasher), ToString(StrongHasher));
 	UNSYNC_LOG_INDENT;
 
 	UNSYNC_LOG(L"Generating data");
@@ -631,7 +631,7 @@ TestPerfComputeBlocksVariable(EWeakHashAlgorithmID WeakHasher, EStrongHashAlgori
 void
 TestBuzhash()
 {
-	UNSYNC_LOG(L"test_buzhash()");
+	UNSYNC_LOG(L"TestBuzhash()");
 	UNSYNC_LOG_INDENT;
 
 	// sanity check
@@ -673,7 +673,7 @@ TestBuzhash()
 void
 TestBasicHash()
 {
-	UNSYNC_LOG(L"test_basic_hash()");
+	UNSYNC_LOG(L"TestBasicHash()");
 	UNSYNC_LOG_INDENT;
 
 	{
@@ -713,7 +713,7 @@ TestBasicHash()
 void
 TestBuffer()
 {
-	UNSYNC_LOG(L"test_buffer()");
+	UNSYNC_LOG(L"TestBuffer()");
 	UNSYNC_LOG_INDENT;
 
 	FBuffer Buf1;
@@ -762,7 +762,7 @@ TestBuffer()
 void
 TestMisc()
 {
-	UNSYNC_LOG(L"test_misc()");
+	UNSYNC_LOG(L"TestMisc()");
 	UNSYNC_LOG_INDENT;
 
 	{
