unsync - Remove std::filesystem namespace alias and add wrappers for various filesystem functions

diff --git a/Engine/Source/Programs/Unsync/Private/UnsyncCmdDiff.h b/Engine/Source/Programs/Unsync/Private/UnsyncCmdDiff.h
--- a/Engine/Source/Programs/Unsync/Private/UnsyncCmdDiff.h
+++ b/Engine/Source/Programs/Unsync/Private/UnsyncCmdDiff.h
@@ -9,9 +9,9 @@ namespace unsync {
 
 struct FCmdDiffOptions
 {
-	fs::path			   Source;
-	fs::path			   Base;
-	fs::path			   Output;
+	FPath				   Source;
+	FPath				   Base;
+	FPath				   Output;
 	int32				   CompressionLevel = 3;
 	uint32				   BlockSize		= 0;
 	EWeakHashAlgorithmID   WeakHasher		= EWeakHashAlgorithmID::BuzHash;
diff --git a/Engine/Source/Programs/Unsync/Private/UnsyncCmdHash.cpp b/Engine/Source/Programs/Unsync/Private/UnsyncCmdHash.cpp
--- a/Engine/Source/Programs/Unsync/Private/UnsyncCmdHash.cpp
+++ b/Engine/Source/Programs/Unsync/Private/UnsyncCmdHash.cpp
@@ -14,16 +14,15 @@ CmdHash(const FCmdHashOptions& Options)
 	{
 		UNSYNC_VERBOSE(L"'%ls' is a directory", Options.Input.wstring().c_str());
 
-		fs::path InputRoot	  = Options.Input;
-		fs::path ManifestRoot = InputRoot / ".unsync";
-		fs::path DirectoryManifestPath;
+		FPath InputRoot	   = Options.Input;
+		FPath ManifestRoot = InputRoot / ".unsync";
+		FPath DirectoryManifestPath;
 
 		if (Options.Output.empty())
 		{
 			DirectoryManifestPath = ManifestRoot / "manifest.bin";
 
-			bool bDirectoryExists =
-				(fs::exists(ManifestRoot) && fs::is_directory(ManifestRoot)) || (GDryRun || fs::create_directories(ManifestRoot));
+			bool bDirectoryExists = (PathExists(ManifestRoot) && IsDirectory(ManifestRoot)) || (GDryRun || CreateDirectories(ManifestRoot));
 
 			if (!bDirectoryExists && !GDryRun)
 			{
@@ -66,7 +65,7 @@ CmdHash(const FCmdHashOptions& Options)
 			UNSYNC_VERBOSE(L"Computing blocks for '%ls' (%.2f MB)", Options.Input.wstring().c_str(), SizeMb(OverlappedFile.GetSize()));
 			FGenericBlockArray GenericBlocks = ComputeBlocks(OverlappedFile, Options.BlockSize, Options.Algorithm);
 
-			fs::path OutputFilename = Options.Output;
+			FPath OutputFilename = Options.Output;
 			if (OutputFilename.empty())
 			{
 				OutputFilename = Options.Input.wstring() + std::wstring(L".unsync");
diff --git a/Engine/Source/Programs/Unsync/Private/UnsyncCmdHash.h b/Engine/Source/Programs/Unsync/Private/UnsyncCmdHash.h
--- a/Engine/Source/Programs/Unsync/Private/UnsyncCmdHash.h
+++ b/Engine/Source/Programs/Unsync/Private/UnsyncCmdHash.h
@@ -9,8 +9,8 @@ namespace unsync {
 
 struct FCmdHashOptions
 {
-	fs::path		  Input;
-	fs::path		  Output;
+	FPath			  Input;
+	FPath			  Output;
 	bool			  bForce	   = false;
 	bool			  bIncremental = false;
 	uint32			  BlockSize	   = uint32(64_KB);
diff --git a/Engine/Source/Programs/Unsync/Private/UnsyncCmdPatch.h b/Engine/Source/Programs/Unsync/Private/UnsyncCmdPatch.h
--- a/Engine/Source/Programs/Unsync/Private/UnsyncCmdPatch.h
+++ b/Engine/Source/Programs/Unsync/Private/UnsyncCmdPatch.h
@@ -9,9 +9,9 @@ namespace unsync {
 
 struct FCmdPatchOptions
 {
-	fs::path Base;
-	fs::path Patch;
-	fs::path Output;
+	FPath Base;
+	FPath Patch;
+	FPath Output;
 };
 
 int32 CmdPatch(const FCmdPatchOptions& Options);
diff --git a/Engine/Source/Programs/Unsync/Private/UnsyncCmdPush.cpp b/Engine/Source/Programs/Unsync/Private/UnsyncCmdPush.cpp
--- a/Engine/Source/Programs/Unsync/Private/UnsyncCmdPush.cpp
+++ b/Engine/Source/Programs/Unsync/Private/UnsyncCmdPush.cpp
@@ -41,7 +41,7 @@ CmdPush(const FCmdPushOptions& Options)
 	AlgorithmOptions.StrongHashAlgorithmId = EStrongHashAlgorithmID::Blake3_128;
 
 	FDirectoryManifest Manifest;
-	fs::path		   ManifestPath	  = Options.Input / ".unsync" / "manifest.bin";	 // TODO: allow manifest path override
+	FPath			   ManifestPath	  = Options.Input / ".unsync" / "manifest.bin";	 // TODO: allow manifest path override
 	bool			   bManifestValid = LoadOrCreateDirectoryManifest(Manifest, Options.Input, BlockSize, AlgorithmOptions);
 
 	if (!bManifestValid)
diff --git a/Engine/Source/Programs/Unsync/Private/UnsyncCmdPush.h b/Engine/Source/Programs/Unsync/Private/UnsyncCmdPush.h
--- a/Engine/Source/Programs/Unsync/Private/UnsyncCmdPush.h
+++ b/Engine/Source/Programs/Unsync/Private/UnsyncCmdPush.h
@@ -10,7 +10,7 @@ namespace unsync {
 
 struct FCmdPushOptions
 {
-	fs::path	Input;
+	FPath		Input;
 	FRemoteDesc Remote;
 };
 
diff --git a/Engine/Source/Programs/Unsync/Private/UnsyncCmdSync.cpp b/Engine/Source/Programs/Unsync/Private/UnsyncCmdSync.cpp
--- a/Engine/Source/Programs/Unsync/Private/UnsyncCmdSync.cpp
+++ b/Engine/Source/Programs/Unsync/Private/UnsyncCmdSync.cpp
@@ -10,13 +10,13 @@ int32  // TODO: return a TResult
 CmdSync(const FCmdSyncOptions& Options)
 {
 	std::error_code ErrorCode	   = {};
-	fs::path		ResolvedSource = Options.Filter ? Options.Filter->Resolve(Options.Source) : Options.Source;
+	FPath			ResolvedSource = Options.Filter ? Options.Filter->Resolve(Options.Source) : Options.Source;
 
 	UNSYNC_VERBOSE(L"Sync source: '%ls'", Options.Source.wstring().c_str());
 	UNSYNC_VERBOSE(L"-- resolved: '%ls'", ResolvedSource.wstring().c_str());
 	UNSYNC_VERBOSE(L"Sync target: '%ls'", Options.Target.wstring().c_str());
 
-	const bool bSourcePathExists	 = fs::exists(ResolvedSource, ErrorCode);
+	const bool bSourcePathExists	 = PathExists(ResolvedSource, ErrorCode);
 	const bool bSourceIsDirectory	 = bSourcePathExists && unsync::IsDirectory(ResolvedSource);
 	const bool bSourceIsManifestHash = !bSourcePathExists && LooksLikeHash160(Options.Source.native());
 
diff --git a/Engine/Source/Programs/Unsync/Private/UnsyncCmdSync.h b/Engine/Source/Programs/Unsync/Private/UnsyncCmdSync.h
--- a/Engine/Source/Programs/Unsync/Private/UnsyncCmdSync.h
+++ b/Engine/Source/Programs/Unsync/Private/UnsyncCmdSync.h
@@ -12,9 +12,9 @@ struct FCmdSyncOptions
 {
 	FAlgorithmOptions Algorithm;
 
-	fs::path Source;
-	fs::path Target;
-	fs::path SourceManifestOverride;
+	FPath Source;
+	FPath Target;
+	FPath SourceManifestOverride;
 
 	FRemoteDesc Remote;
 
diff --git a/Engine/Source/Programs/Unsync/Private/UnsyncCommon.h b/Engine/Source/Programs/Unsync/Private/UnsyncCommon.h
--- a/Engine/Source/Programs/Unsync/Private/UnsyncCommon.h
+++ b/Engine/Source/Programs/Unsync/Private/UnsyncCommon.h
@@ -20,7 +20,7 @@ using int16 = int16_t;
 using int32 = int32_t;
 using int64 = int64_t;
 
-namespace fs = std::filesystem;
+using FPath = std::filesystem::path;
 
 #define UNSYNC_DISALLOW_COPY_ASSIGN(T) \
 	T(const T&) = delete;              \
diff --git a/Engine/Source/Programs/Unsync/Private/UnsyncCore.cpp b/Engine/Source/Programs/Unsync/Private/UnsyncCore.cpp
--- a/Engine/Source/Programs/Unsync/Private/UnsyncCore.cpp
+++ b/Engine/Source/Programs/Unsync/Private/UnsyncCore.cpp
@@ -1772,7 +1772,7 @@ IsSynchronized(const FNeedList& NeedList, const FGenericBlockArray& SourceBlocks
 }
 
 FDirectoryManifest
-CreateDirectoryManifest(const fs::path& Root, uint32 BlockSize, FAlgorithmOptions Algorithm)
+CreateDirectoryManifest(const FPath& Root, uint32 BlockSize, FAlgorithmOptions Algorithm)
 {
 	UNSYNC_LOG_INDENT;
 
@@ -1780,7 +1780,7 @@ CreateDirectoryManifest(const fs::path& Root, uint32 BlockSize, FAlgorithmOption
 
 	Result.Options = Algorithm;
 
-	fs::path ManifestRootPath = Root / ".unsync";
+	FPath ManifestRootPath = Root / ".unsync";
 
 	FTimePoint TimeBegin = TimePointNow();
 
@@ -1790,7 +1790,7 @@ CreateDirectoryManifest(const fs::path& Root, uint32 BlockSize, FAlgorithmOption
 
 	std::mutex ResultMutex;
 
-	for (const fs::directory_entry& Dir : fs::recursive_directory_iterator(Root))
+	for (const std::filesystem::directory_entry& Dir : std::filesystem::recursive_directory_iterator(Root))
 	{
 		if (Dir.is_directory())
 		{
@@ -1802,7 +1802,7 @@ CreateDirectoryManifest(const fs::path& Root, uint32 BlockSize, FAlgorithmOption
 			continue;
 		}
 
-		fs::path RelativePath = Dir.path().lexically_relative(Root);
+		FPath RelativePath = Dir.path().lexically_relative(Root);
 
 		UNSYNC_VERBOSE2(L"Found '%ls'", RelativePath.wstring().c_str());
 
@@ -1822,8 +1822,8 @@ CreateDirectoryManifest(const fs::path& Root, uint32 BlockSize, FAlgorithmOption
 
 		if (BlockSize)
 		{
-			fs::path FilePath = Root / RelativePath;
-			auto	 File	  = std::make_shared<NativeFile>(FilePath, EFileMode::ReadOnlyUnbuffered);
+			FPath FilePath = Root / RelativePath;
+			auto  File	   = std::make_shared<NativeFile>(FilePath, EFileMode::ReadOnlyUnbuffered);
 			if (File->IsValid())
 			{
 				UNSYNC_VERBOSE(L"Computing blocks for '%ls' (%.2f MB)", FilePath.wstring().c_str(), double(File->GetSize()) / (1 << 20));
@@ -1892,10 +1892,10 @@ CreateDirectoryManifest(const fs::path& Root, uint32 BlockSize, FAlgorithmOption
 }
 
 FDirectoryManifest
-CreateDirectoryManifestIncremental(const fs::path& Root, uint32 BlockSize, FAlgorithmOptions Algorithm)
+CreateDirectoryManifestIncremental(const FPath& Root, uint32 BlockSize, FAlgorithmOptions Algorithm)
 {
-	fs::path ManifestRoot		   = Root / ".unsync";
-	fs::path DirectoryManifestPath = ManifestRoot / "manifest.bin";
+	FPath ManifestRoot			= Root / ".unsync";
+	FPath DirectoryManifestPath = ManifestRoot / "manifest.bin";
 
 	FDirectoryManifest OldManifest;
 	const bool		   bExistingManifestLoaded = LoadDirectoryManifest(OldManifest, Root, DirectoryManifestPath);
@@ -1937,7 +1937,7 @@ CreateDirectoryManifestIncremental(const fs::path& Root, uint32 BlockSize, FAlgo
 }
 
 void
-UpdateDirectoryManifestBlocks(FDirectoryManifest& Result, const fs::path& Root, uint32 BlockSize, FAlgorithmOptions Algorithm)
+UpdateDirectoryManifestBlocks(FDirectoryManifest& Result, const FPath& Root, uint32 BlockSize, FAlgorithmOptions Algorithm)
 {
 	UNSYNC_LOG_INDENT;
 
@@ -1968,8 +1968,8 @@ UpdateDirectoryManifestBlocks(FDirectoryManifest& Result, const fs::path& Root,
 
 		++NumProcessedFiles;
 
-		fs::path FilePath = Root / It.first;
-		auto	 File	  = std::make_shared<NativeFile>(FilePath, EFileMode::ReadOnlyUnbuffered);
+		FPath FilePath = Root / It.first;
+		auto  File	   = std::make_shared<NativeFile>(FilePath, EFileMode::ReadOnlyUnbuffered);
 		if (File->IsValid())
 		{
 			UNSYNC_VERBOSE(L"Computing blocks for '%ls' (%.2f MB)", FilePath.wstring().c_str(), double(File->GetSize()) / (1 << 20));
@@ -2041,17 +2041,17 @@ UpdateDirectoryManifestBlocks(FDirectoryManifest& Result, const fs::path& Root,
 }
 
 bool
-LoadOrCreateDirectoryManifest(FDirectoryManifest& Result, const fs::path& Root, uint32 BlockSize, FAlgorithmOptions Algorithm)
+LoadOrCreateDirectoryManifest(FDirectoryManifest& Result, const FPath& Root, uint32 BlockSize, FAlgorithmOptions Algorithm)
 {
 	UNSYNC_LOG_INDENT;
 
-	fs::path ManifestRoot		   = Root / ".unsync";
-	fs::path DirectoryManifestPath = ManifestRoot / "manifest.bin";
+	FPath ManifestRoot			= Root / ".unsync";
+	FPath DirectoryManifestPath = ManifestRoot / "manifest.bin";
 
 	FDirectoryManifest OldDirectoryManifest;
 	FDirectoryManifest NewDirectoryManifest;
 
-	const bool bManifestFileExists = fs::exists(DirectoryManifestPath);
+	const bool bManifestFileExists = PathExists(DirectoryManifestPath);
 	if (!bManifestFileExists)
 	{
 		UNSYNC_VERBOSE(L"Manifest file '%ls' does not exist", DirectoryManifestPath.wstring().c_str());
@@ -2333,10 +2333,10 @@ ValidateTarget(FIOReader& Reader, const FNeedList& NeedList, EStrongHashAlgorith
 
 FFileSyncResult
 SyncFile(const FNeedList&		   NeedList,
-		 const fs::path&		   SourceFilePath,
+		 const FPath&			   SourceFilePath,
 		 const FGenericBlockArray& SourceBlocks,
 		 FIOReader&				   BaseDataReader,
-		 const fs::path&		   TargetFilePath,
+		 const FPath&			   TargetFilePath,
 		 const FSyncFileOptions&   Options)
 {
 	UNSYNC_LOG_INDENT;
@@ -2359,10 +2359,10 @@ SyncFile(const FNeedList&		   NeedList,
 		}
 		else
 		{
-			fs::path TargetFileParent = TargetFilePath.parent_path();
-			if (!fs::exists(TargetFileParent))
+			FPath TargetFileParent = TargetFilePath.parent_path();
+			if (!PathExists(TargetFileParent))
 			{
-				fs::create_directories(TargetFileParent);
+				CreateDirectories(TargetFileParent);
 			}
 
 			auto TargetFile = NativeFile(TargetFilePath, EFileMode::CreateWriteOnly, 0);
@@ -2381,7 +2381,7 @@ SyncFile(const FNeedList&		   NeedList,
 	{
 		LogStatus(TargetFilePath.wstring().c_str(), L"Initializing");
 
-		fs::path TempTargetFilePath = TargetFilePath;
+		FPath TempTargetFilePath = TargetFilePath;
 		TempTargetFilePath.replace_extension(TargetFilePath.extension().wstring() + L".tmp");
 		const FNeedListSize TargetFileSizeInfo = ComputeNeedListSize(NeedList);
 
@@ -2401,10 +2401,10 @@ SyncFile(const FNeedList&		   NeedList,
 		}
 		else
 		{
-			fs::path TargetFileParent = TempTargetFilePath.parent_path();
-			if (!fs::exists(TargetFileParent))
+			FPath TargetFileParent = TempTargetFilePath.parent_path();
+			if (!PathExists(TargetFileParent))
 			{
-				fs::create_directories(TargetFileParent);
+				CreateDirectories(TargetFileParent);
 			}
 
 			TargetFile = std::make_unique<NativeFile>(TempTargetFilePath, EFileMode::CreateWriteOnly, TargetFileSizeInfo.TotalBytes);
@@ -2474,7 +2474,7 @@ SyncFile(const FNeedList&		   NeedList,
 			}
 
 			std::error_code Ec = {};
-			fs::rename(TempTargetFilePath, TargetFilePath, Ec);
+			FileRename(TempTargetFilePath, TargetFilePath, Ec);
 
 			if (Ec.value() == 0)
 			{
@@ -2507,10 +2507,10 @@ SyncFile(const FNeedList&		   NeedList,
 }
 
 FFileSyncResult
-SyncFile(const fs::path&		   SourceFilePath,
+SyncFile(const FPath&			   SourceFilePath,
 		 const FGenericBlockArray& SourceBlocks,
 		 FIOReader&				   BaseDataReader,
-		 const fs::path&		   TargetFilePath,
+		 const FPath&			   TargetFilePath,
 		 const FSyncFileOptions&   Options)
 {
 	UNSYNC_LOG_INDENT;
@@ -2526,7 +2526,7 @@ SyncFile(const fs::path&		   SourceFilePath,
 }
 
 FFileSyncResult
-SyncFile(const fs::path& SourceFilePath, const fs::path& BaseFilePath, const fs::path& TargetFilePath, const FSyncFileOptions& InOptions)
+SyncFile(const FPath& SourceFilePath, const FPath& BaseFilePath, const FPath& TargetFilePath, const FSyncFileOptions& InOptions)
 {
 	UNSYNC_LOG_INDENT;
 
@@ -2541,7 +2541,7 @@ SyncFile(const fs::path& SourceFilePath, const fs::path& BaseFilePath, const fs:
 
 		UNSYNC_VERBOSE(L"Full copy required for '%ls' (base does not exist)", BaseFilePath.wstring().c_str());
 		std::error_code ErrorCode;
-		bool			bCopyOk = fs::copy_file(SourceFilePath, TargetFilePath, ErrorCode);
+		bool			bCopyOk = FileCopy(SourceFilePath, TargetFilePath, ErrorCode);
 		if (bCopyOk)
 		{
 			Result.Status = EFileSyncStatus::Ok;
@@ -2556,7 +2556,7 @@ SyncFile(const fs::path& SourceFilePath, const fs::path& BaseFilePath, const fs:
 	}
 
 	FGenericBlockArray SourceBlocks;
-	fs::path		   BlockFilename = BaseFilePath.wstring() + std::wstring(L".unsync");
+	FPath			   BlockFilename = BaseFilePath.wstring() + std::wstring(L".unsync");
 
 	UNSYNC_VERBOSE(L"Loading block manifest from '%ls'", BlockFilename.wstring().c_str());
 	if (LoadBlocks(SourceBlocks, Options.BlockSize, BlockFilename.c_str()))
@@ -2568,7 +2568,7 @@ SyncFile(const fs::path& SourceFilePath, const fs::path& BaseFilePath, const fs:
 		UNSYNC_VERBOSE(L"Full copy required (manifest file does not exist or is invalid)");
 
 		std::error_code ErrorCode;
-		bool			bCopyOk = fs::copy_file(SourceFilePath, TargetFilePath, ErrorCode);
+		bool			bCopyOk = FileCopy(SourceFilePath, TargetFilePath, ErrorCode);
 		if (bCopyOk)
 		{
 			Result.Status = EFileSyncStatus::Ok;
@@ -2588,21 +2588,21 @@ SyncFile(const fs::path& SourceFilePath, const fs::path& BaseFilePath, const fs:
 }
 
 std::error_code
-CopyFileIfNewer(const fs::path& Source, const fs::path& Target)
+CopyFileIfNewer(const FPath& Source, const FPath& Target)
 {
 	FileAttributes	SourceAttr = GetFileAttrib(Source);
 	FileAttributes	TargetAttr = GetFileAttrib(Target);
 	std::error_code Ec;
 	if (SourceAttr.Size != TargetAttr.Size || SourceAttr.Mtime != TargetAttr.Mtime)
 	{
-		fs::copy_file(Source, Target, fs::copy_options::overwrite_existing, Ec);
+		FileCopyOverwrite(Source, Target, Ec);
 	}
 	return Ec;
 }
 
 // Delete files from target directory that are not in the source directory manifest
 static void
-DeleteUnnecessaryFiles(const fs::path&			 TargetDirectory,
+DeleteUnnecessaryFiles(const FPath&				 TargetDirectory,
 					   const FDirectoryManifest& TargetDirectoryManifest,
 					   const FDirectoryManifest& ReferenceManifest)
 {
@@ -2611,7 +2611,7 @@ DeleteUnnecessaryFiles(const fs::path&			 TargetDirectory,
 		const std::wstring& TargetFileName = TargetManifestEntry.first;
 		if (ReferenceManifest.Files.find(TargetFileName) == ReferenceManifest.Files.end())
 		{
-			fs::path FilePath = TargetDirectory / TargetFileName;
+			FPath FilePath = TargetDirectory / TargetFileName;
 			if (GDryRun)
 			{
 				UNSYNC_VERBOSE(L"Deleting '%ls' (skipped due to dry run mode)", FilePath.wstring().c_str());
@@ -2620,7 +2620,7 @@ DeleteUnnecessaryFiles(const fs::path&			 TargetDirectory,
 			{
 				UNSYNC_VERBOSE(L"Deleting '%ls'", FilePath.wstring().c_str());
 				std::error_code ErrorCode = {};
-				fs::remove(FilePath, ErrorCode);
+				FileRemove(FilePath, ErrorCode);
 				if (ErrorCode)
 				{
 					UNSYNC_VERBOSE(L"System error code %d: %hs", ErrorCode.value(), ErrorCode.message().c_str());
@@ -2630,16 +2630,16 @@ DeleteUnnecessaryFiles(const fs::path&			 TargetDirectory,
 	}
 }
 
-fs::path
+FPath
 ToPath(const std::wstring_view& Str)
 {
 #if UNSYNC_PLATFORM_UNIX
 	// TODO: ensure that all serialized path separators are unix style ('/')
 	std::wstring Temp = std::wstring(Str);
 	std::replace(Temp.begin(), Temp.end(), L'\\', L'/');
-	return fs::path(Temp);
+	return FPath(Temp);
 #else	// UNSYNC_PLATFORM_UNIX
-	return fs::path(Str);
+	return FPath(Str);
 #endif	// UNSYNC_PLATFORM_UNIX
 }
 
@@ -2651,10 +2651,10 @@ SyncDirectory(const FSyncDirectoryOptions& SyncOptions)
 
 	UNSYNC_ASSERT(bFileSystemSource || bServerSource);
 
-	const fs::path	SourcePath			   = bFileSystemSource ? fs::absolute(SyncOptions.Source) : SyncOptions.Source;
-	const fs::path	BasePath			   = fs::absolute(SyncOptions.Base);
-	const fs::path	TargetPath			   = fs::absolute(SyncOptions.Target);
-	const fs::path& SourceManifestOverride = SyncOptions.SourceManifestOverride;
+	const FPath	 SourcePath				= bFileSystemSource ? std::filesystem::absolute(SyncOptions.Source) : SyncOptions.Source;
+	const FPath	 BasePath				= std::filesystem::absolute(SyncOptions.Base);
+	const FPath	 TargetPath				= std::filesystem::absolute(SyncOptions.Target);
+	const FPath& SourceManifestOverride = SyncOptions.SourceManifestOverride;
 
 	FSyncFilter* SyncFilter = SyncOptions.SyncFilter;
 
@@ -2667,14 +2667,14 @@ SyncDirectory(const FSyncDirectoryOptions& SyncOptions)
 		UNSYNC_VERBOSE(L"Unnecessary files will be deleted after sync (cleanup mode)");
 	}
 
-	fs::path BaseManifestRoot = BasePath / ".unsync";
-	fs::path BaseManifestPath = BaseManifestRoot / "manifest.bin";
+	FPath BaseManifestRoot = BasePath / ".unsync";
+	FPath BaseManifestPath = BaseManifestRoot / "manifest.bin";
 
-	fs::path TargetManifestRoot = TargetPath / ".unsync";
-	fs::path TargetManifestPath = TargetManifestRoot / "manifest.bin";
-	fs::path TargetTempPath		= TargetManifestRoot / "temp";
+	FPath TargetManifestRoot = TargetPath / ".unsync";
+	FPath TargetManifestPath = TargetManifestRoot / "manifest.bin";
+	FPath TargetTempPath	 = TargetManifestRoot / "temp";
 
-	bool bTempDirectoryExists = (fs::exists(TargetTempPath) && fs::is_directory(TargetTempPath)) || fs::create_directories(TargetTempPath);
+	bool bTempDirectoryExists = (PathExists(TargetTempPath) && IsDirectory(TargetTempPath)) || CreateDirectories(TargetTempPath);
 
 	if (!bTempDirectoryExists)
 	{
@@ -2682,11 +2682,11 @@ SyncDirectory(const FSyncDirectoryOptions& SyncOptions)
 		return false;
 	}
 
-	fs::path	  LogFilePath = TargetManifestRoot / L"unsync.log";
+	FPath		  LogFilePath = TargetManifestRoot / L"unsync.log";
 	FLogFileScope LogFileScope(LogFilePath.wstring().c_str());
 	SetCrashDumpPath(TargetManifestRoot);
 
-	auto ShouldSync = [SyncFilter](const fs::path& Filename) -> bool {
+	auto ShouldSync = [SyncFilter](const FPath& Filename) -> bool {
 		if (SyncFilter)
 		{
 			return SyncFilter->ShouldSync(Filename);
@@ -2697,18 +2697,18 @@ SyncDirectory(const FSyncDirectoryOptions& SyncOptions)
 		}
 	};
 
-	auto ResolvePath = [SyncFilter](const fs::path& Filename) -> fs::path { return SyncFilter ? SyncFilter->Resolve(Filename) : Filename; };
+	auto ResolvePath = [SyncFilter](const FPath& Filename) -> FPath { return SyncFilter ? SyncFilter->Resolve(Filename) : Filename; };
 
 	FRemoteDesc ProxySettings = SyncOptions.Remote ? *SyncOptions.Remote : FRemoteDesc();
 	FProxyPool	ProxyPool(ProxySettings);
 
 	FDirectoryManifest SourceDirectoryManifest;
-	fs::path		   SourceManifestTempPath;
+	FPath			   SourceManifestTempPath;
 
 	if (bFileSystemSource || !SourceManifestOverride.empty())
 	{
-		fs::path SourceManifestRoot = SourcePath / ".unsync";
-		fs::path SourceManifestPath = SourceManifestRoot / "manifest.bin";
+		FPath SourceManifestRoot = SourcePath / ".unsync";
+		FPath SourceManifestPath = SourceManifestRoot / "manifest.bin";
 
 		SourceManifestPath = ResolvePath(SourceManifestPath);
 
@@ -2723,7 +2723,7 @@ SyncDirectory(const FSyncDirectoryOptions& SyncOptions)
 		std::string SourcePathHashStr = BytesToHexString(SourcePathHash.Data, sizeof(SourcePathHash.Data));
 		SourceManifestTempPath		  = TargetTempPath / SourcePathHashStr;
 
-		if (!fs::exists(SourceManifestPath))
+		if (!PathExists(SourceManifestPath))
 		{
 			UNSYNC_ERROR(L"Source manifest '%ls' does not exist", SourceManifestPath.wstring().c_str());
 			return false;
@@ -2771,7 +2771,7 @@ SyncDirectory(const FSyncDirectoryOptions& SyncOptions)
 		{
 			FMemReader		Reader(*ManifestBuffer);
 			FIOReaderStream Stream(Reader);
-			fs::path		EmptyRoot;	// Don't have a sensible path when not using file system as source
+			FPath			EmptyRoot;	// Don't have a sensible path when not using file system as source
 			bSourceManifestOk = LoadDirectoryManifest(SourceDirectoryManifest, EmptyRoot, Stream);
 		}
 		else
@@ -2811,11 +2811,11 @@ SyncDirectory(const FSyncDirectoryOptions& SyncOptions)
 	{
 		const FFileManifest* SourceManifest = nullptr;
 		const FFileManifest* BaseManifest	= nullptr;
-		fs::path			 OriginalSourceFilePath;
-		fs::path			 ResolvedSourceFilePath;
-		fs::path			 BaseFilePath;
-		fs::path			 TargetFilePath;
-		fs::path			 RelativeFilePath;
+		FPath				 OriginalSourceFilePath;
+		FPath				 ResolvedSourceFilePath;
+		FPath				 BaseFilePath;
+		FPath				 TargetFilePath;
+		FPath				 RelativeFilePath;
 		FNeedList			 NeedList;
 
 		uint64 NeedBytesFromSource = 0;
@@ -2889,11 +2889,11 @@ SyncDirectory(const FSyncDirectoryOptions& SyncOptions)
 			}
 		}
 
-		fs::path SourceFilePath = SourcePath / ToPath(SourceManifestIt.first);
-		fs::path BaseFilePath	= BasePath / ToPath(SourceManifestIt.first);
-		fs::path TargetFilePath = TargetPath / ToPath(SourceManifestIt.first);
+		FPath SourceFilePath = SourcePath / ToPath(SourceManifestIt.first);
+		FPath BaseFilePath	 = BasePath / ToPath(SourceManifestIt.first);
+		FPath TargetFilePath = TargetPath / ToPath(SourceManifestIt.first);
 
-		fs::path ResolvedSourceFilePath = ResolvePath(SourceFilePath);
+		FPath ResolvedSourceFilePath = ResolvePath(SourceFilePath);
 
 		if (bFileSystemSource && SyncOptions.bValidateSourceFiles)
 		{
@@ -2942,7 +2942,7 @@ SyncDirectory(const FSyncDirectoryOptions& SyncOptions)
 				StatPartialCopy++;
 
 				if (bFileSystemSource && SyncOptions.bValidateSourceFiles && !SourceAttribCache.Exists(ResolvedSourceFilePath) &&
-					!fs::exists(ResolvedSourceFilePath))
+					!PathExists(ResolvedSourceFilePath))
 				{
 					UNSYNC_VERBOSE(L"Source file '%ls' does not exist", SourceFilePath.wstring().c_str());
 					continue;
@@ -3003,7 +3003,7 @@ SyncDirectory(const FSyncDirectoryOptions& SyncOptions)
 
 			const FGenericBlockArray& SourceBlocks = Item.SourceManifest->Blocks;
 
-			if (Item.IsBaseValid() && fs::exists(Item.BaseFilePath))
+			if (Item.IsBaseValid() && PathExists(Item.BaseFilePath))
 			{
 				NativeFile BaseFile(Item.BaseFilePath, EFileMode::ReadOnlyUnbuffered);
 				uint32	   SourceBlockSize = Item.SourceManifest->BlockSize;
@@ -3150,7 +3150,7 @@ SyncDirectory(const FSyncDirectoryOptions& SyncOptions)
 	{
 		struct BackgroundTaskResult
 		{
-			fs::path		TargetFilePath;
+			FPath			TargetFilePath;
 			FFileSyncResult SyncResult;
 			bool			bIsPartialCopy = false;
 		};
@@ -3341,7 +3341,7 @@ SyncDirectory(const FSyncDirectoryOptions& SyncOptions)
 		// TODO: cache the manfiest if it's downloaded from Jupiter also
 
 		std::error_code ErrorCode;
-		bool bCopyOk = fs::copy_file(SourceManifestTempPath, TargetManifestPath, fs::copy_options::overwrite_existing, ErrorCode);
+		bool			bCopyOk = FileCopyOverwrite(SourceManifestTempPath, TargetManifestPath, ErrorCode);
 
 		if (!bCopyOk)
 		{
@@ -3356,12 +3356,12 @@ SyncDirectory(const FSyncDirectoryOptions& SyncOptions)
 	return bSyncSucceeded;
 }
 
-fs::path
-FSyncFilter::Resolve(const fs::path& Filename) const
+FPath
+FSyncFilter::Resolve(const FPath& Filename) const
 {
 	std::wstring FilenameLower = StringToLower(Filename.wstring());
 
-	fs::path Result;
+	FPath Result;
 	for (const auto& Alias : DfsAliases)
 	{
 		// TODO: add a case-insensitive find() helper
@@ -3733,7 +3733,7 @@ FSyncFilter::FSyncFilter(const std::wstring& InExcludedWords)
 }
 
 bool
-FSyncFilter::ShouldSync(const fs::path& Filename) const
+FSyncFilter::ShouldSync(const FPath& Filename) const
 {
 #if UNSYNC_PLATFORM_WINDOWS
 	return ShouldSync(Filename.native());
@@ -3891,10 +3891,10 @@ LogManifestDiff(ELogLevel LogLevel, const FDirectoryManifest& ManifestA, const F
 }
 
 int32
-CmdInfo(const fs::path& InputA, const fs::path& InputB)
+CmdInfo(const FPath& InputA, const FPath& InputB)
 {
-	fs::path DirectoryManifestPathA = InputA / ".unsync" / "manifest.bin";
-	fs::path DirectoryManifestPathB = InputB / ".unsync" / "manifest.bin";
+	FPath DirectoryManifestPathA = InputA / ".unsync" / "manifest.bin";
+	FPath DirectoryManifestPathB = InputB / ".unsync" / "manifest.bin";
 
 	FDirectoryManifest ManifestA;
 	bool			   bManifestAValid = LoadDirectoryManifest(ManifestA, InputA, DirectoryManifestPathA);
diff --git a/Engine/Source/Programs/Unsync/Private/UnsyncCore.h b/Engine/Source/Programs/Unsync/Private/UnsyncCore.h
--- a/Engine/Source/Programs/Unsync/Private/UnsyncCore.h
+++ b/Engine/Source/Programs/Unsync/Private/UnsyncCore.h
@@ -120,7 +120,7 @@ struct FFileManifest
 	// runtime-only data:
 	FGenericBlockArray Blocks;
 	FGenericBlockArray MacroBlocks;
-	fs::path		   CurrentPath;
+	FPath			   CurrentPath;
 };
 
 struct FAlgorithmOptionsV5
@@ -304,24 +304,24 @@ struct FSyncFileOptions
 	bool bValidateTargetFiles = true;  // WARNING: turning this off is intended only for testing/profiling
 };
 
-FFileSyncResult SyncFile(const fs::path&		   SourceFilePath,
+FFileSyncResult SyncFile(const FPath&			   SourceFilePath,
 						 const FGenericBlockArray& SourceBlocks,
 						 FIOReader&				   BaseDataReader,
-						 const fs::path&		   TargetFilePath,
+						 const FPath&			   TargetFilePath,
 						 const FSyncFileOptions&   Options);
 
-FFileSyncResult SyncFile(const fs::path&		 SourceFilePath,
-						 const fs::path&		 BaseFilePath,
-						 const fs::path&		 TargetFilePath,
+FFileSyncResult SyncFile(const FPath&			 SourceFilePath,
+						 const FPath&			 BaseFilePath,
+						 const FPath&			 TargetFilePath,
 						 const FSyncFileOptions& Options);
 
 struct FSyncFilter
 {
 	FSyncFilter(const std::wstring& ExcludedWords);
 
-	bool	 ShouldSync(const fs::path& Filename) const;
-	bool	 ShouldSync(const std::wstring& Filename) const;
-	fs::path Resolve(const fs::path& Filename) const;
+	bool  ShouldSync(const FPath& Filename) const;
+	bool  ShouldSync(const std::wstring& Filename) const;
+	FPath Resolve(const FPath& Filename) const;
 
 	std::vector<std::wstring> ExcludedWords;
 	std::vector<FDfsAlias>	  DfsAliases;
@@ -337,10 +337,10 @@ enum class ESyncSourceType
 struct FSyncDirectoryOptions
 {
 	ESyncSourceType	   SourceType;
-	fs::path		   Source;					// remote data location
-	fs::path		   Target;					// output target location
-	fs::path		   Base;					// base data location, which typically is the same as sync target
-	fs::path		   SourceManifestOverride;	// force the manifest to be read from a specified file instead of source directory
+	FPath			   Source;					// remote data location
+	FPath			   Target;					// output target location
+	FPath			   Base;					// base data location, which typically is the same as sync target
+	FPath			   SourceManifestOverride;	// force the manifest to be read from a specified file instead of source directory
 	FSyncFilter*	   SyncFilter = nullptr;	// filter callback for partial sync support
 	const FRemoteDesc* Remote	  = nullptr;	// unsync proxy server connection settings
 	bool			   bCleanup	  = false;	// whether to cleanup any files in the target directory that are not in the source manifest file
@@ -352,10 +352,10 @@ struct FSyncDirectoryOptions
 
 bool SyncDirectory(const FSyncDirectoryOptions& SyncOptions);
 
-void UpdateDirectoryManifestBlocks(FDirectoryManifest& Result, const fs::path& Root, uint32 BlockSize, FAlgorithmOptions Algorithm);
-FDirectoryManifest CreateDirectoryManifest(const fs::path& Root, uint32 BlockSize, FAlgorithmOptions Algorithm);
-FDirectoryManifest CreateDirectoryManifestIncremental(const fs::path& Root, uint32 BlockSize, FAlgorithmOptions Algorithm);
-bool LoadOrCreateDirectoryManifest(FDirectoryManifest& Result, const fs::path& Root, uint32 BlockSize, FAlgorithmOptions Algorithm);
+void UpdateDirectoryManifestBlocks(FDirectoryManifest& Result, const FPath& Root, uint32 BlockSize, FAlgorithmOptions Algorithm);
+FDirectoryManifest CreateDirectoryManifest(const FPath& Root, uint32 BlockSize, FAlgorithmOptions Algorithm);
+FDirectoryManifest CreateDirectoryManifestIncremental(const FPath& Root, uint32 BlockSize, FAlgorithmOptions Algorithm);
+bool LoadOrCreateDirectoryManifest(FDirectoryManifest& Result, const FPath& Root, uint32 BlockSize, FAlgorithmOptions Algorithm);
 
 // Computes a Blake3 hash of the manifest blocks and file metadata, ignoring any other metadata.
 // Files are processed in sorted order, with file names treated as utf-8.
@@ -371,6 +371,6 @@ FHash160 ComputeSerializedManifestHash160(const FDirectoryManifest& Manifest);
 FBlock128			   ToBlock128(const FGenericBlock& GenericBlock);
 std::vector<FBlock128> ToBlock128(FGenericBlockArray& GenericBlocks);
 
-int32 CmdInfo(const fs::path& InputA, const fs::path& InputB);
+int32 CmdInfo(const FPath& InputA, const FPath& InputB);
 
 }  // namespace unsync
diff --git a/Engine/Source/Programs/Unsync/Private/UnsyncFile.cpp b/Engine/Source/Programs/Unsync/Private/UnsyncFile.cpp
--- a/Engine/Source/Programs/Unsync/Private/UnsyncFile.cpp
+++ b/Engine/Source/Programs/Unsync/Private/UnsyncFile.cpp
@@ -66,7 +66,7 @@ struct FCreateFileInfo
 	}
 };
 
-FWindowsFile::FWindowsFile(const fs::path& InFilename, EFileMode InMode, uint64 InSize) : Mode(InMode)
+FWindowsFile::FWindowsFile(const FPath& InFilename, EFileMode InMode, uint64 InSize) : Mode(InMode)
 {
 	Filename = InFilename;
 
@@ -517,7 +517,7 @@ FWindowsFile::FlushOne()
 }
 
 FileAttributes
-GetFileAttrib(const fs::path& Path, FFileAttributeCache* AttribCache)
+GetFileAttrib(const FPath& Path, FFileAttributeCache* AttribCache)
 {
 	FileAttributes Result;
 
@@ -546,7 +546,7 @@ GetFileAttrib(const fs::path& Path, FFileAttributeCache* AttribCache)
 }
 
 bool
-SetFileMtime(const fs::path& Path, uint64 Mtime)
+SetFileMtime(const FPath& Path, uint64 Mtime)
 {
 	UNSYNC_ASSERT(!GDryRun);
 	HANDLE Fh = CreateFileW(Path.c_str(), FILE_WRITE_ATTRIBUTES, 0, nullptr, OPEN_EXISTING, FILE_ATTRIBUTE_NORMAL, nullptr);
@@ -566,7 +566,7 @@ SetFileMtime(const fs::path& Path, uint64 Mtime)
 }
 
 bool
-SetFileReadOnly(const fs::path& Path, bool bReadOnly)
+SetFileReadOnly(const FPath& Path, bool bReadOnly)
 {
 	UNSYNC_ASSERT(!GDryRun);
 
@@ -593,14 +593,14 @@ SetFileReadOnly(const fs::path& Path, bool bReadOnly)
 }
 
 uint64
-ToWindowsFileTime(const fs::file_time_type& T)
+ToWindowsFileTime(const std::filesystem::file_time_type& T)
 {
 	return T.time_since_epoch().count();
 }
 #endif	// UNSYNC_PLATFORM_WINDOWS
 
 #if UNSYNC_PLATFORM_UNIX
-FUnixFile::FUnixFile(const fs::path& InFilename, EFileMode InMode, uint64 in_size) : Filename(InFilename), Mode(InMode)
+FUnixFile::FUnixFile(const FPath& InFilename, EFileMode InMode, uint64 in_size) : Filename(InFilename), Mode(InMode)
 {
 	FileHandle = fopen(InFilename.native().c_str(), IsReadOnly(Mode) ? "rb" : "w+b");
 	if (FileHandle == nullptr)
@@ -690,7 +690,7 @@ FUnixFile::Write(const void* data, uint64 DestOffset, uint64 WriteSize)
 }
 
 FileAttributes
-GetFileAttrib(const fs::path& path, FFileAttributeCache* AttribCache)
+GetFileAttrib(const FPath& path, FFileAttributeCache* AttribCache)
 {
 	FileAttributes result;
 
@@ -707,7 +707,7 @@ GetFileAttrib(const fs::path& path, FFileAttributeCache* AttribCache)
 	// TODO: could potentially use std::filesystem::directory_entry for this on all platforms
 
 	std::error_code ec	  = {};
-	auto			entry = fs::directory_entry(path, ec);
+	auto			entry = std::filesystem::directory_entry(path, ec);
 
 	if (ec.value() == 0)
 	{
@@ -722,21 +722,21 @@ GetFileAttrib(const fs::path& path, FFileAttributeCache* AttribCache)
 }
 
 bool
-SetFileReadOnly(const fs::path& path, bool bReadOnly)
+SetFileReadOnly(const FPath& path, bool bReadOnly)
 {
 	UNSYNC_WARNING(L"SetFileReadOnly() is not implemented");
 	return false;
 }
 
 bool
-SetFileMtime(const fs::path& path, uint64 mtime)
+SetFileMtime(const FPath& path, uint64 mtime)
 {
 	UNSYNC_WARNING(L"SetFileMtime() is not implemented");
 	return false;
 }
 
 uint64
-ToWindowsFileTime(const fs::file_time_type& FileTime)
+ToWindowsFileTime(const std::filesystem::file_time_type& FileTime)
 {
 	const uint64 NANOS_PER_TICK	  = 100ull;
 	const uint64 TICKS_PER_SECOND = 1'000'000'000ull / NANOS_PER_TICK;	// each tick is 100ns
@@ -762,7 +762,7 @@ ToWindowsFileTime(const fs::file_time_type& FileTime)
 #endif	// UNSYNC_PLATFORM_UNIX
 
 FBuffer
-ReadFileToBuffer(const fs::path& Filename)
+ReadFileToBuffer(const FPath& Filename)
 {
 	FBuffer	   Result;
 	NativeFile File(Filename, EFileMode::ReadOnly);
@@ -776,7 +776,7 @@ ReadFileToBuffer(const fs::path& Filename)
 }
 
 bool
-WriteBufferToFile(const fs::path& Filename, const uint8* Data, uint64 Size)
+WriteBufferToFile(const FPath& Filename, const uint8* Data, uint64 Size)
 {
 	UNSYNC_LOG_INDENT;
 	UNSYNC_ASSERT(Data);
@@ -798,7 +798,7 @@ WriteBufferToFile(const fs::path& Filename, const uint8* Data, uint64 Size)
 }
 
 bool
-WriteBufferToFile(const fs::path& Filename, const FBuffer& Buffer)
+WriteBufferToFile(const FPath& Filename, const FBuffer& Buffer)
 {
 	return WriteBufferToFile(Filename, Buffer.Data(), Buffer.Size());
 }
@@ -957,7 +957,7 @@ GetCurrentIoCacheSize()
 }
 
 FFileAttributeCache
-CreateFileAttributeCache(const fs::path& Root, const FSyncFilter* SyncFilter)
+CreateFileAttributeCache(const FPath& Root, const FSyncFilter* SyncFilter)
 {
 	FFileAttributeCache Result;
 
@@ -972,9 +972,9 @@ CreateFileAttributeCache(const fs::path& Root, const FSyncFilter* SyncFilter)
 		}
 	};
 
-	fs::path ResolvedRoot = SyncFilter ? SyncFilter->Resolve(Root) : Root;
+	FPath ResolvedRoot = SyncFilter ? SyncFilter->Resolve(Root) : Root;
 
-	for (const fs::directory_entry& Dir : fs::recursive_directory_iterator(ResolvedRoot))
+	for (const std::filesystem::directory_entry& Dir : std::filesystem::recursive_directory_iterator(ResolvedRoot))
 	{
 		if (Dir.is_directory())
 		{
@@ -1003,12 +1003,55 @@ CreateFileAttributeCache(const fs::path& Root, const FSyncFilter* SyncFilter)
 }
 
 bool
-IsDirectory(const fs::path& Path)
+IsDirectory(const FPath& Path)
 {
 	FileAttributes Attr = GetFileAttrib(Path);
 	return Attr.bValid && Attr.bDirectory;
 }
 
+bool
+PathExists(const FPath& Path)
+{
+	return std::filesystem::exists(Path);
+}
+
+bool
+PathExists(const FPath& Path, std::error_code& OutErrorCode)
+{
+	return std::filesystem::exists(Path, OutErrorCode);
+}
+
+bool
+CreateDirectories(const FPath& Path)
+{
+	return std::filesystem::create_directories(Path);
+}
+
+bool
+FileRename(const FPath& From, const FPath& To, std::error_code& OutErrorCode)
+{
+	std::filesystem::rename(From, To, OutErrorCode);
+	return OutErrorCode.value() == 0;
+}
+
+bool
+FileCopy(const FPath& From, const FPath& To, std::error_code& OutErrorCode)
+{
+	return std::filesystem::copy_file(From, To, OutErrorCode);
+}
+
+bool
+FileCopyOverwrite(const FPath& From, const FPath& To, std::error_code& OutErrorCode)
+{
+	return std::filesystem::copy_file(From, To, std::filesystem::copy_options::overwrite_existing, OutErrorCode);
+}
+
+bool
+FileRemove(const FPath& Path, std::error_code& OutErrorCode)
+{
+	return std::filesystem::remove(Path, OutErrorCode);
+}
+
 FMemReader::FMemReader(const uint8* InData, uint64 InDataSize) : Data(InData), Size(InDataSize)
 {
 }
diff --git a/Engine/Source/Programs/Unsync/Private/UnsyncFile.h b/Engine/Source/Programs/Unsync/Private/UnsyncFile.h
--- a/Engine/Source/Programs/Unsync/Private/UnsyncFile.h
+++ b/Engine/Source/Programs/Unsync/Private/UnsyncFile.h
@@ -158,7 +158,7 @@ struct FIOReaderWriter : FIOReader, FIOWriter
 #if UNSYNC_PLATFORM_WINDOWS
 struct FWindowsFile : FIOReaderWriter
 {
-	FWindowsFile(const fs::path& Filename, EFileMode Mode = EFileMode::ReadOnly, uint64 InSize = 0);
+	FWindowsFile(const FPath& Filename, EFileMode Mode = EFileMode::ReadOnly, uint64 InSize = 0);
 	~FWindowsFile();
 
 	// IOBase
@@ -194,7 +194,7 @@ struct FWindowsFile : FIOReaderWriter
 	};
 	Command Commands[NUM_QUEUES] = {};
 
-	fs::path Filename;
+	FPath Filename;
 
 	static constexpr uint32 UNBUFFERED_READ_ALIGNMENT = 4096;
 
@@ -211,7 +211,7 @@ using NativeFile = FWindowsFile;
 #if UNSYNC_PLATFORM_UNIX
 struct FUnixFile : FIOReaderWriter
 {
-	FUnixFile(const fs::path& InFilename, EFileMode InMode = EFileMode::ReadOnly, uint64 InSize = 0);
+	FUnixFile(const FPath& InFilename, EFileMode InMode = EFileMode::ReadOnly, uint64 InSize = 0);
 	~FUnixFile();
 
 	// IOBase
@@ -232,7 +232,7 @@ struct FUnixFile : FIOReaderWriter
 	uint64 FileSize	 = 0;
 	int32  LastError = 0;
 
-	fs::path Filename;
+	FPath Filename;
 
 	static constexpr uint32 UNBUFFERED_READ_ALIGNMENT = 4096;
 
@@ -421,9 +421,9 @@ struct FIOReaderStream
 	uint64	   Offset = 0;
 };
 
-FBuffer ReadFileToBuffer(const fs::path& Filename);
-bool	WriteBufferToFile(const fs::path& Filename, const uint8* Data, uint64 Size);
-bool	WriteBufferToFile(const fs::path& Filename, const FBuffer& Buffer);
+FBuffer ReadFileToBuffer(const FPath& Filename);
+bool	WriteBufferToFile(const FPath& Filename, const uint8* Data, uint64 Size);
+bool	WriteBufferToFile(const FPath& Filename, const FBuffer& Buffer);
 
 struct FileAttributes
 {
@@ -436,19 +436,26 @@ struct FileAttributes
 
 struct FFileAttributeCache
 {
-	std::unordered_map<fs::path::string_type, FileAttributes> Map;
+	std::unordered_map<FPath::string_type, FileAttributes> Map;
 
-	const bool Exists(const fs::path& Path) const;
+	const bool Exists(const FPath& Path) const;
 };
 
-FileAttributes GetFileAttrib(const fs::path& Path, FFileAttributeCache* AttribCache = nullptr);
-bool		   SetFileMtime(const fs::path& Path, uint64 Mtime);
-bool		   SetFileReadOnly(const fs::path& Path, bool ReadOnly);
-bool		   IsDirectory(const fs::path& Path);
+FileAttributes GetFileAttrib(const FPath& Path, FFileAttributeCache* AttribCache = nullptr);
+bool		   SetFileMtime(const FPath& Path, uint64 Mtime);
+bool		   SetFileReadOnly(const FPath& Path, bool ReadOnly);
+bool		   IsDirectory(const FPath& Path);
+bool		   PathExists(const FPath& Path);
+bool		   PathExists(const FPath& Path, std::error_code& OutErrorCode);
+bool		   CreateDirectories(const FPath& Path);
+bool		   FileRename(const FPath& From, const FPath& To, std::error_code& OutErrorCode);
+bool		   FileCopy(const FPath& From, const FPath& To, std::error_code& OutErrorCode);
+bool		   FileCopyOverwrite(const FPath& From, const FPath& To, std::error_code& OutErrorCode);
+bool		   FileRemove(const FPath& Path, std::error_code& OutErrorCode);
 
-uint64 ToWindowsFileTime(const fs::file_time_type& T);
+uint64 ToWindowsFileTime(const std::filesystem::file_time_type& T);
 
 struct FSyncFilter;
-FFileAttributeCache CreateFileAttributeCache(const fs::path& Root, const FSyncFilter* SyncFilter = nullptr);
+FFileAttributeCache CreateFileAttributeCache(const FPath& Root, const FSyncFilter* SyncFilter = nullptr);
 
 }  // namespace unsync
diff --git a/Engine/Source/Programs/Unsync/Private/UnsyncLog.cpp b/Engine/Source/Programs/Unsync/Private/UnsyncLog.cpp
--- a/Engine/Source/Programs/Unsync/Private/UnsyncLog.cpp
+++ b/Engine/Source/Programs/Unsync/Private/UnsyncLog.cpp
@@ -320,16 +320,16 @@ LogError(const FError& E)
 	}
 }
 
-static fs::path GCrashDumpPath;
+static FPath GCrashDumpPath;
 void
-SetCrashDumpPath(const fs::path& Path)
+SetCrashDumpPath(const FPath& Path)
 {
 	std::lock_guard<std::mutex> LockGuard(GLogMutex);
 	GCrashDumpPath = Path;
 }
 
 bool
-get_crash_dump_path(fs::path& OutPath)
+get_crash_dump_path(FPath& OutPath)
 {
 	std::lock_guard<std::mutex> LockGuard(GLogMutex);
 	if (GCrashDumpPath.empty())
@@ -350,7 +350,7 @@ LogWriteCrashDump(void* InExceptionPointers)
 
 	_EXCEPTION_POINTERS* ExceptionPointers = (_EXCEPTION_POINTERS*)InExceptionPointers;
 
-	fs::path CrashDumpPath;
+	FPath CrashDumpPath;
 	if (!get_crash_dump_path(CrashDumpPath))
 	{
 		return false;
@@ -359,7 +359,7 @@ LogWriteCrashDump(void* InExceptionPointers)
 	wchar_t DumpFilename[4096] = {};
 	uint64	Timestamp		   = TimePointNow().time_since_epoch().count();
 	swprintf_s(DumpFilename, L"unsync-%016llX.dmp", Timestamp);
-	fs::path CrashDumpFilename = CrashDumpPath / fs::path(DumpFilename);
+	FPath CrashDumpFilename = CrashDumpPath / FPath(DumpFilename);
 
 	LogPrintf(ELogLevel::Info, L"!!! Writing crash dump to '%ls'\n", CrashDumpFilename.wstring().c_str());
 
diff --git a/Engine/Source/Programs/Unsync/Private/UnsyncLog.h b/Engine/Source/Programs/Unsync/Private/UnsyncLog.h
--- a/Engine/Source/Programs/Unsync/Private/UnsyncLog.h
+++ b/Engine/Source/Programs/Unsync/Private/UnsyncLog.h
@@ -84,8 +84,8 @@ LogGlobalStatus(const wchar_t* Status)
 	LogStatus(nullptr, Status);
 }
 
-void SetCrashDumpPath(const fs::path& Path);
-bool GetCrashDumpPath(fs::path& OutPath);
+void SetCrashDumpPath(const FPath& Path);
+bool GetCrashDumpPath(FPath& OutPath);
 bool LogWriteCrashDump(void* ExceptionPointers);
 
 #define UNSYNC_LOG_INDENT		FLogIndentScope UNSYNC_CONCAT(log_indent_, __LINE__);
diff --git a/Engine/Source/Programs/Unsync/Private/UnsyncMain.cpp b/Engine/Source/Programs/Unsync/Private/UnsyncMain.cpp
--- a/Engine/Source/Programs/Unsync/Private/UnsyncMain.cpp
+++ b/Engine/Source/Programs/Unsync/Private/UnsyncMain.cpp
@@ -276,7 +276,7 @@ InnerMain(int Argc, char** Argv)
 	{
 		// Derive remote server address from source name if explicit --proxy or --remote option is not provided for sync command.
 
-		if (Cli.got_subcommand(SubSync) && !fs::exists(SourceFilenameUtf8))
+		if (Cli.got_subcommand(SubSync) && !PathExists(SourceFilenameUtf8))
 		{
 			TResult<FRemoteDesc> ParsedRemoteDesc = FRemoteDesc::FromUrl(SourceFilenameUtf8);
 			if (ParsedRemoteDesc.IsOk())
@@ -320,14 +320,14 @@ InnerMain(int Argc, char** Argv)
 		}
 	}
 
-	fs::path InputFilename			= NormalizeFilenameUtf8(InputFilenameUtf8);
-	fs::path InputFilename2			= NormalizeFilenameUtf8(InputFilename2Utf8);
-	fs::path OutputFilename			= NormalizeFilenameUtf8(OutputFilenameUtf8);
-	fs::path BaseFilename			= NormalizeFilenameUtf8(BaseFilenameUtf8);
-	fs::path SourceFilename			= NormalizeFilenameUtf8(SourceFilenameUtf8);
-	fs::path TargetFilename			= NormalizeFilenameUtf8(TargetFilenameUtf8);
-	fs::path PatchFilename			= NormalizeFilenameUtf8(PatchFilenameUtf8);
-	fs::path SourceManifestFilename = NormalizeFilenameUtf8(SourceManifestFilenameUtf8);
+	FPath InputFilename			 = NormalizeFilenameUtf8(InputFilenameUtf8);
+	FPath InputFilename2		 = NormalizeFilenameUtf8(InputFilename2Utf8);
+	FPath OutputFilename		 = NormalizeFilenameUtf8(OutputFilenameUtf8);
+	FPath BaseFilename			 = NormalizeFilenameUtf8(BaseFilenameUtf8);
+	FPath SourceFilename		 = NormalizeFilenameUtf8(SourceFilenameUtf8);
+	FPath TargetFilename		 = NormalizeFilenameUtf8(TargetFilenameUtf8);
+	FPath PatchFilename			 = NormalizeFilenameUtf8(PatchFilenameUtf8);
+	FPath SourceManifestFilename = NormalizeFilenameUtf8(SourceManifestFilenameUtf8);
 
 	UNSYNC_VERBOSE(L"UNSYNC %hs", GetVersionString().c_str());
 
@@ -458,7 +458,7 @@ InnerMain(int Argc, char** Argv)
 
 			FDfsAlias DfsAlias;
 			DfsAlias.Source = DfsEntries.Root;
-			DfsAlias.Target = fs::path(L"\\\\") / FoundDfsStorage->Server / FoundDfsStorage->Share;
+			DfsAlias.Target = FPath(L"\\\\") / FoundDfsStorage->Server / FoundDfsStorage->Share;
 
 			UNSYNC_VERBOSE(L"Using DFS alias '%ls' -> '%ls'", DfsAlias.Source.wstring().c_str(), DfsAlias.Target.wstring().c_str());
 
@@ -471,8 +471,8 @@ InnerMain(int Argc, char** Argv)
 
 	if (!HttpHeaderFilenameUtf8.empty())
 	{
-		fs::path Filename		   = NormalizeFilenameUtf8(HttpHeaderFilenameUtf8);
-		FBuffer	 HttpHeadersBuffer = ReadFileToBuffer(Filename);
+		FPath	Filename		  = NormalizeFilenameUtf8(HttpHeaderFilenameUtf8);
+		FBuffer HttpHeadersBuffer = ReadFileToBuffer(Filename);
 
 		const uint8 Bom[2] = {0xFF, 0xFE};
 		if (HttpHeadersBuffer.Size() > 2 && !memcmp(HttpHeadersBuffer.Data(), Bom, 2))
@@ -504,8 +504,8 @@ InnerMain(int Argc, char** Argv)
 
 	if (!CacertFilenameUtf8.empty())
 	{
-		fs::path CacertPath	  = NormalizeFilenameUtf8(CacertFilenameUtf8);
-		FBuffer	 CacertBuffer = ReadFileToBuffer(CacertPath);
+		FPath	CacertPath	 = NormalizeFilenameUtf8(CacertFilenameUtf8);
+		FBuffer CacertBuffer = ReadFileToBuffer(CacertPath);
 
 		RemoteDesc.TlsCacert = std::make_shared<FBuffer>(std::move(CacertBuffer));
 	}
@@ -585,7 +585,7 @@ InnerMain(int Argc, char** Argv)
 	else if (Cli.got_subcommand(SubQuery))
 	{
 		FCmdQueryOptions QueryOptions;
-		QueryOptions.Query  = QueryStringUtf8;
+		QueryOptions.Query	= QueryStringUtf8;
 		QueryOptions.Remote = RemoteDesc;
 		return CmdQuery(QueryOptions);
 	}
diff --git a/Engine/Source/Programs/Unsync/Private/UnsyncProxy.cpp b/Engine/Source/Programs/Unsync/Private/UnsyncProxy.cpp
--- a/Engine/Source/Programs/Unsync/Private/UnsyncProxy.cpp
+++ b/Engine/Source/Programs/Unsync/Private/UnsyncProxy.cpp
@@ -346,7 +346,7 @@ FUnsyncProtocolImpl::~FUnsyncProtocolImpl()
 }
 
 void
-FBlockRequestMap::AddFileBlocks(const fs::path& OriginalFilePath, const fs::path& ResolvedFilePath, const FFileManifest& FileManifest)
+FBlockRequestMap::AddFileBlocks(const FPath& OriginalFilePath, const FPath& ResolvedFilePath, const FFileManifest& FileManifest)
 {
 	UNSYNC_ASSERTF(StrongHasher != EStrongHashAlgorithmID::Invalid, L"Request map is not initialized");
 
@@ -505,7 +505,7 @@ FProxyPool::IsValid() const
 }
 
 void
-FProxyPool::BuildFileBlockRequests(const fs::path& OriginalFilePath, const fs::path& ResolvedFilePath, const FFileManifest& FileManifest)
+FProxyPool::BuildFileBlockRequests(const FPath& OriginalFilePath, const FPath& ResolvedFilePath, const FFileManifest& FileManifest)
 {
 	std::lock_guard<std::mutex> LockGuard(Mutex);
 	RequestMap.AddFileBlocks(OriginalFilePath, ResolvedFilePath, FileManifest);
diff --git a/Engine/Source/Programs/Unsync/Private/UnsyncProxy.h b/Engine/Source/Programs/Unsync/Private/UnsyncProxy.h
--- a/Engine/Source/Programs/Unsync/Private/UnsyncProxy.h
+++ b/Engine/Source/Programs/Unsync/Private/UnsyncProxy.h
@@ -57,7 +57,7 @@ public:
 		StrongHasher = InStrongHasher;
 	}
 
-	void AddFileBlocks(const fs::path& OriginalFilePath, const fs::path& ResolvedFilePath, const FFileManifest& Manifest);
+	void AddFileBlocks(const FPath& OriginalFilePath, const FPath& ResolvedFilePath, const FFileManifest& Manifest);
 
 	const std::vector<std::string>& GetFileList() const { return FileListUtf8; }
 	const FBlockRequest*			FindRequest(const FGenericHash& BlockHash) const;
@@ -129,7 +129,7 @@ public:
 	const FRemoteDesc RemoteDesc;
 
 	void InitRequestMap(EStrongHashAlgorithmID InStrongHasher);
-	void BuildFileBlockRequests(const fs::path& OriginalFilePath, const fs::path& ResolvedFilePath, const FFileManifest& FileManifest);
+	void BuildFileBlockRequests(const FPath& OriginalFilePath, const FPath& ResolvedFilePath, const FFileManifest& FileManifest);
 
 private:
 	std::vector<std::unique_ptr<FProxy>> Pool;
diff --git a/Engine/Source/Programs/Unsync/Private/UnsyncSerialization.cpp b/Engine/Source/Programs/Unsync/Private/UnsyncSerialization.cpp
--- a/Engine/Source/Programs/Unsync/Private/UnsyncSerialization.cpp
+++ b/Engine/Source/Programs/Unsync/Private/UnsyncSerialization.cpp
@@ -116,7 +116,7 @@ ToGenericBlock(const std::vector<BlockType>& Blocks, EHashType HashType)
 }
 
 bool
-SaveBlocks(const std::vector<FBlock128>& Blocks, uint32 BlockSize, const fs::path& Filename)
+SaveBlocks(const std::vector<FBlock128>& Blocks, uint32 BlockSize, const FPath& Filename)
 {
 	const uint64 OutputSize = sizeof(FBlockFileHeader) + sizeof(FBlock128) * Blocks.size();
 
@@ -148,7 +148,7 @@ SaveBlocks(const std::vector<FBlock128>& Blocks, uint32 BlockSize, const fs::pat
 }
 
 bool
-LoadBlocks(FGenericBlockArray& OutBlocks, uint32& OutBlockSize, const fs::path& Filename)
+LoadBlocks(FGenericBlockArray& OutBlocks, uint32& OutBlockSize, const FPath& Filename)
 {
 	UNSYNC_LOG_INDENT;
 
@@ -317,7 +317,7 @@ SaveMacroBlocks(const FDirectoryManifest& Manifest)
 }
 
 bool  // TODO: return a TResult
-LoadDirectoryManifest(FDirectoryManifest& OutManifest, const fs::path& Root, FIOReaderStream& Stream)
+LoadDirectoryManifest(FDirectoryManifest& OutManifest, const FPath& Root, FIOReaderStream& Stream)
 {
 	OutManifest = FDirectoryManifest();
 
@@ -498,7 +498,7 @@ LoadDirectoryManifest(FDirectoryManifest& OutManifest, const fs::path& Root, FIO
 }
 
 bool  // TODO: return a TResult
-LoadDirectoryManifest(FDirectoryManifest& OutManifest, const fs::path& Root, const fs::path& ManifestFilename)
+LoadDirectoryManifest(FDirectoryManifest& OutManifest, const FPath& Root, const FPath& ManifestFilename)
 {
 	// TODO: use compact binary to store the manifest
 	UNSYNC_LOG_INDENT;
@@ -624,7 +624,7 @@ SaveDirectoryManifest(const FDirectoryManifest& Manifest, FVectorStreamOut& Stre
 }
 
 bool
-SaveDirectoryManifest(const FDirectoryManifest& Manifest, const fs::path& Filename)
+SaveDirectoryManifest(const FDirectoryManifest& Manifest, const FPath& Filename)
 {
 	FBuffer			 OutputBuffer;
 	FVectorStreamOut OutputStream(OutputBuffer);
diff --git a/Engine/Source/Programs/Unsync/Private/UnsyncSerialization.h b/Engine/Source/Programs/Unsync/Private/UnsyncSerialization.h
--- a/Engine/Source/Programs/Unsync/Private/UnsyncSerialization.h
+++ b/Engine/Source/Programs/Unsync/Private/UnsyncSerialization.h
@@ -9,13 +9,13 @@ namespace unsync {
 struct FIOReaderStream;
 struct FVectorStreamOut;
 
-bool LoadBlocks(FGenericBlockArray& OutBlocks, uint32& OutBlockSize, const fs::path& Filename);
-bool SaveBlocks(const std::vector<FBlock128>& Blocks, uint32 BlockSize, const fs::path& Filename);
+bool LoadBlocks(FGenericBlockArray& OutBlocks, uint32& OutBlockSize, const FPath& Filename);
+bool SaveBlocks(const std::vector<FBlock128>& Blocks, uint32 BlockSize, const FPath& Filename);
 
-bool LoadDirectoryManifest(FDirectoryManifest& OutManifest, const fs::path& Root, FIOReaderStream& Stream);
-bool LoadDirectoryManifest(FDirectoryManifest& OutManifest, const fs::path& Root, const fs::path& Filename);
+bool LoadDirectoryManifest(FDirectoryManifest& OutManifest, const FPath& Root, FIOReaderStream& Stream);
+bool LoadDirectoryManifest(FDirectoryManifest& OutManifest, const FPath& Root, const FPath& Filename);
 
 bool SaveDirectoryManifest(const FDirectoryManifest& Manifest, FVectorStreamOut& Stream);
-bool SaveDirectoryManifest(const FDirectoryManifest& Manifest, const fs::path& Filename);
+bool SaveDirectoryManifest(const FDirectoryManifest& Manifest, const FPath& Filename);
 
 }  // namespace unsync
diff --git a/Engine/Source/Programs/Unsync/Private/UnsyncTest.cpp b/Engine/Source/Programs/Unsync/Private/UnsyncTest.cpp
--- a/Engine/Source/Programs/Unsync/Private/UnsyncTest.cpp
+++ b/Engine/Source/Programs/Unsync/Private/UnsyncTest.cpp
@@ -459,14 +459,14 @@ TestFiles()
 	// TODO: test different chunk modes
 	const EChunkingAlgorithmID ChunkingMode = EChunkingAlgorithmID::FixedBlocks;
 
-	fs::path TempDirPath = fs::temp_directory_path() / "unsync_test";
-	fs::create_directories(TempDirPath);
+	FPath TempDirPath = std::filesystem::temp_directory_path() / "unsync_test";
+	CreateDirectories(TempDirPath);
 
-	bool bDirectoryExists = fs::exists(TempDirPath) && fs::is_directory(TempDirPath);
+	bool bDirectoryExists = PathExists(TempDirPath) && IsDirectory(TempDirPath);
 	UNSYNC_ASSERT(bDirectoryExists);
 
-	fs::path TestFilename = TempDirPath / "a.bin";
-	uint64	 TestFileSize = 32'489'595;
+	FPath  TestFilename = TempDirPath / "a.bin";
+	uint64 TestFileSize = 32'489'595;
 
 	std::unique_ptr<uint8> TempBuffer = std::unique_ptr<uint8>(new uint8[TestFileSize]);
 
diff --git a/Engine/Source/Programs/Unsync/Private/UnsyncUtil.cpp b/Engine/Source/Programs/Unsync/Private/UnsyncUtil.cpp
--- a/Engine/Source/Programs/Unsync/Private/UnsyncUtil.cpp
+++ b/Engine/Source/Programs/Unsync/Private/UnsyncUtil.cpp
@@ -103,7 +103,7 @@ ConvertWideToUtf8(std::wstring_view StringWide)
 }
 
 const bool
-FFileAttributeCache::Exists(const fs::path& Path) const
+FFileAttributeCache::Exists(const FPath& Path) const
 {
 	const auto It = Map.find(Path.native());
 	return It != Map.end();
@@ -118,7 +118,7 @@ StringToLower(const std::wstring& Input)
 }
 
 FDfsMirrorInfo
-DfsEnumerate(const fs::path& Root)
+DfsEnumerate(const FPath& Root)
 {
 	FDfsMirrorInfo Result;
 
@@ -213,7 +213,7 @@ DfsEnumerate(const fs::path& Root)
 	return Result;
 }
 
-fs::path
+FPath
 NormalizeFilenameUtf8(std::string& InFilename)
 {
 	std::string_view Filename	   = InFilename;
@@ -223,7 +223,7 @@ NormalizeFilenameUtf8(std::string& InFilename)
 		Filename = Filename.substr(FileUrlPrefix.length());
 	}
 
-	fs::path FilenameAsPath = ConvertUtf8ToWide(Filename);
+	FPath FilenameAsPath = ConvertUtf8ToWide(Filename);
 	return FilenameAsPath.lexically_normal();
 }
 
diff --git a/Engine/Source/Programs/Unsync/Private/UnsyncUtil.h b/Engine/Source/Programs/Unsync/Private/UnsyncUtil.h
--- a/Engine/Source/Programs/Unsync/Private/UnsyncUtil.h
+++ b/Engine/Source/Programs/Unsync/Private/UnsyncUtil.h
@@ -144,12 +144,12 @@ struct FDfsMirrorInfo
 	std::wstring				 Root;
 	std::vector<FDfsStorageInfo> Storages;
 };
-FDfsMirrorInfo DfsEnumerate(const fs::path& Root);
+FDfsMirrorInfo DfsEnumerate(const FPath& Root);
 
 struct FDfsAlias
 {
-	fs::path Source;
-	fs::path Target;
+	FPath Source;
+	FPath Target;
 };
 
 template<typename T>
@@ -275,7 +275,7 @@ CountLeadingZeros64(uint64 X)
 	return ((63 - XLog2) & Mask) | (64 & ~Mask);
 }
 
-fs::path NormalizeFilenameUtf8(std::string& InFilename);
+FPath NormalizeFilenameUtf8(std::string& InFilename);
 
 const FBuffer& GetSystemRootCerts();
 
