unsync - Handle long paths on Windows, bump version to 1.0.39

diff --git a/Engine/Source/Programs/Unsync/Private/UnsyncCore.cpp b/Engine/Source/Programs/Unsync/Private/UnsyncCore.cpp
--- a/Engine/Source/Programs/Unsync/Private/UnsyncCore.cpp
+++ b/Engine/Source/Programs/Unsync/Private/UnsyncCore.cpp
@@ -24,7 +24,7 @@ UNSYNC_THIRD_PARTY_INCLUDES_START
 #include <flat_hash_map.hpp>
 UNSYNC_THIRD_PARTY_INCLUDES_END
 
-#define UNSYNC_VERSION_STR "1.0.38"
+#define UNSYNC_VERSION_STR "1.0.39"
 
 namespace unsync {
 
@@ -1780,8 +1780,6 @@ CreateDirectoryManifest(const FPath& Root, uint32 BlockSize, FAlgorithmOptions A
 
 	Result.Options = Algorithm;
 
-	FPath ManifestRootPath = Root / ".unsync";
-
 	FTimePoint TimeBegin = TimePointNow();
 
 	FTaskGroup	 TaskGroup;
@@ -1789,21 +1787,22 @@ CreateDirectoryManifest(const FPath& Root, uint32 BlockSize, FAlgorithmOptions A
 	FSemaphore	 Semaphore(MaxConcurrentFiles);
 
 	std::mutex ResultMutex;
+	FPath UnsyncDirName = ".unsync";
 
-	for (const std::filesystem::directory_entry& Dir : std::filesystem::recursive_directory_iterator(Root))
+	for (const std::filesystem::directory_entry& Dir : RecursiveDirectoryScan(Root))
 	{
 		if (Dir.is_directory())
 		{
 			continue;
 		}
 
-		if (Dir.path().native().starts_with(ManifestRootPath.native()))
+		FPath RelativePath = std::filesystem::relative(Dir.path(), Root);
+
+		if (RelativePath.native().starts_with(UnsyncDirName.native()))
 		{
 			continue;
 		}
 
-		FPath RelativePath = Dir.path().lexically_relative(Root);
-
 		UNSYNC_VERBOSE2(L"Found '%ls'", RelativePath.wstring().c_str());
 
 		FFileManifest FileManifest;
@@ -2846,7 +2845,8 @@ SyncDirectory(const FSyncDirectoryOptions& SyncOptions)
 	bool			   bBaseDirectoryManifestValid = false;
 	bool			   bQuickDifferencePossible	   = false;
 
-	if (SyncOptions.bQuickDifference && SourceDirectoryManifest.Options.ChunkingAlgorithmId == EChunkingAlgorithmID::VariableBlocks)
+	if (SyncOptions.bQuickDifference && SourceDirectoryManifest.Options.ChunkingAlgorithmId == EChunkingAlgorithmID::VariableBlocks
+		&& PathExists(BaseManifestPath))
 	{
 		bBaseDirectoryManifestValid = LoadDirectoryManifest(BaseDirectoryManifest, BasePath, BaseManifestPath);
 		if (bBaseDirectoryManifestValid &&
@@ -3345,7 +3345,9 @@ SyncDirectory(const FSyncDirectoryOptions& SyncOptions)
 
 		if (!bCopyOk)
 		{
-			UNSYNC_WARNING(L"Failed to save manifest after sync. System error code: %d.", ErrorCode.value());
+			UNSYNC_WARNING(L"Failed to save manifest after sync. System error code: %hs (%d).",
+						   ErrorCode.message().c_str(),
+						   ErrorCode.value());
 		}
 	}
 
diff --git a/Engine/Source/Programs/Unsync/Private/UnsyncFile.cpp b/Engine/Source/Programs/Unsync/Private/UnsyncFile.cpp
--- a/Engine/Source/Programs/Unsync/Private/UnsyncFile.cpp
+++ b/Engine/Source/Programs/Unsync/Private/UnsyncFile.cpp
@@ -24,6 +24,33 @@ namespace unsync {
 
 bool GForceBufferedFiles = false;
 
+// Returns extended absolute path of a form \\?\D:\verylongpath or \\?\UNC\servername\verylongpath
+// Expects an absolute path input.
+// https://docs.microsoft.com/en-us/windows/win32/fileio/maximum-file-path-limitation
+static FPath
+MakeExtendedAbsolutePath(const FPath& InAbsolutePath)
+{
+	UNSYNC_ASSERT(InAbsolutePath.is_absolute());
+
+#if UNSYNC_PLATFORM_WINDOWS
+	const std::wstring& InFilenameString = InAbsolutePath.native();
+	if (InFilenameString.starts_with(L"\\\\?\\"))
+	{
+		return InAbsolutePath;
+	}
+	else if (InFilenameString.starts_with(L"\\\\"))
+	{
+		return std::wstring(L"\\\\?\\UNC\\") + InFilenameString.substr(2);
+	}
+	else
+	{
+		return std::wstring(L"\\\\?\\") + InFilenameString;
+	}
+#else // UNSYNC_PLATFORM_WINDOWS
+	return InAbsolutePath;
+#endif // UNSYNC_PLATFORM_WINDOWS
+}
+
 #if UNSYNC_PLATFORM_WINDOWS
 inline uint64
 MakeU64(FILETIME Ft)
@@ -68,7 +95,7 @@ struct FCreateFileInfo
 
 FWindowsFile::FWindowsFile(const FPath& InFilename, EFileMode InMode, uint64 InSize) : Mode(InMode)
 {
-	Filename = InFilename;
+	Filename = MakeExtendedAbsolutePath(InFilename);
 
 	bool bOpenedOk = OpenFileHandle(InMode);
 
@@ -531,8 +558,10 @@ GetFileAttrib(const FPath& Path, FFileAttributeCache* AttribCache)
 		}
 	}
 
+	FPath ExtendedPath = MakeExtendedAbsolutePath(Path);
+
 	WIN32_FILE_ATTRIBUTE_DATA AttributeData;
-	BOOL					  Ok = GetFileAttributesExW(Path.c_str(), GetFileExInfoStandard, &AttributeData);
+	BOOL					  Ok = GetFileAttributesExW(ExtendedPath.c_str(), GetFileExInfoStandard, &AttributeData);
 	if (Ok)
 	{
 		Result.bDirectory = !!(AttributeData.dwFileAttributes & FILE_ATTRIBUTE_DIRECTORY);
@@ -549,7 +578,10 @@ bool
 SetFileMtime(const FPath& Path, uint64 Mtime)
 {
 	UNSYNC_ASSERT(!GDryRun);
-	HANDLE Fh = CreateFileW(Path.c_str(), FILE_WRITE_ATTRIBUTES, 0, nullptr, OPEN_EXISTING, FILE_ATTRIBUTE_NORMAL, nullptr);
+
+	FPath ExtendedPath = MakeExtendedAbsolutePath(Path);
+
+	HANDLE Fh = CreateFileW(ExtendedPath.c_str(), FILE_WRITE_ATTRIBUTES, 0, nullptr, OPEN_EXISTING, FILE_ATTRIBUTE_NORMAL, nullptr);
 	if (Fh != INVALID_HANDLE_VALUE)
 	{
 		FILETIME C, A, w;
@@ -570,7 +602,9 @@ SetFileReadOnly(const FPath& Path, bool bReadOnly)
 {
 	UNSYNC_ASSERT(!GDryRun);
 
-	uint32 OldAttributes = GetFileAttributesW(Path.c_str());
+	FPath ExtendedPath = MakeExtendedAbsolutePath(Path);
+
+	uint32 OldAttributes = GetFileAttributesW(ExtendedPath.c_str());
 	uint32 NewAttributes = OldAttributes;
 
 	if (bReadOnly)
@@ -588,7 +622,7 @@ SetFileReadOnly(const FPath& Path, bool bReadOnly)
 	}
 	else
 	{
-		return SetFileAttributesW(Path.c_str(), NewAttributes);
+		return SetFileAttributesW(ExtendedPath.c_str(), NewAttributes);
 	}
 }
 
@@ -974,7 +1008,7 @@ CreateFileAttributeCache(const FPath& Root, const FSyncFilter* SyncFilter)
 
 	FPath ResolvedRoot = SyncFilter ? SyncFilter->Resolve(Root) : Root;
 
-	for (const std::filesystem::directory_entry& Dir : std::filesystem::recursive_directory_iterator(ResolvedRoot))
+	for (const std::filesystem::directory_entry& Dir : RecursiveDirectoryScan(ResolvedRoot))
 	{
 		if (Dir.is_directory())
 		{
@@ -1012,44 +1046,60 @@ IsDirectory(const FPath& Path)
 bool
 PathExists(const FPath& Path)
 {
-	return std::filesystem::exists(Path);
+	FPath ExtendedPath = MakeExtendedAbsolutePath(Path);
+	return std::filesystem::exists(ExtendedPath);
 }
 
 bool
 PathExists(const FPath& Path, std::error_code& OutErrorCode)
 {
-	return std::filesystem::exists(Path, OutErrorCode);
+	FPath ExtendedPath = MakeExtendedAbsolutePath(Path);
+	return std::filesystem::exists(ExtendedPath, OutErrorCode);
 }
 
 bool
 CreateDirectories(const FPath& Path)
 {
-	return std::filesystem::create_directories(Path);
+	FPath ExtendedPath = MakeExtendedAbsolutePath(Path);
+	return std::filesystem::create_directories(ExtendedPath);
 }
 
 bool
 FileRename(const FPath& From, const FPath& To, std::error_code& OutErrorCode)
 {
-	std::filesystem::rename(From, To, OutErrorCode);
+	FPath ExtendedFrom = MakeExtendedAbsolutePath(From);
+	FPath ExtendedTo   = MakeExtendedAbsolutePath(To);
+	std::filesystem::rename(ExtendedFrom, ExtendedTo, OutErrorCode);
 	return OutErrorCode.value() == 0;
 }
 
 bool
 FileCopy(const FPath& From, const FPath& To, std::error_code& OutErrorCode)
 {
-	return std::filesystem::copy_file(From, To, OutErrorCode);
+	FPath ExtendedFrom = MakeExtendedAbsolutePath(From);
+	FPath ExtendedTo   = MakeExtendedAbsolutePath(To);
+	return std::filesystem::copy_file(ExtendedFrom, ExtendedTo, OutErrorCode);
 }
 
 bool
 FileCopyOverwrite(const FPath& From, const FPath& To, std::error_code& OutErrorCode)
 {
-	return std::filesystem::copy_file(From, To, std::filesystem::copy_options::overwrite_existing, OutErrorCode);
+	FPath ExtendedFrom = MakeExtendedAbsolutePath(From);
+	FPath ExtendedTo   = MakeExtendedAbsolutePath(To);
+	return std::filesystem::copy_file(ExtendedFrom, ExtendedTo, std::filesystem::copy_options::overwrite_existing, OutErrorCode);
 }
 
 bool
 FileRemove(const FPath& Path, std::error_code& OutErrorCode)
 {
-	return std::filesystem::remove(Path, OutErrorCode);
+	FPath ExtendedPath = MakeExtendedAbsolutePath(Path);
+	return std::filesystem::remove(ExtendedPath, OutErrorCode);
+}
+
+std::filesystem::recursive_directory_iterator RecursiveDirectoryScan(const FPath& Path)
+{
+	FPath ExtendedPath = MakeExtendedAbsolutePath(Path);
+	return std::filesystem::recursive_directory_iterator(ExtendedPath);
 }
 
 FMemReader::FMemReader(const uint8* InData, uint64 InDataSize) : Data(InData), Size(InDataSize)
diff --git a/Engine/Source/Programs/Unsync/Private/UnsyncFile.h b/Engine/Source/Programs/Unsync/Private/UnsyncFile.h
--- a/Engine/Source/Programs/Unsync/Private/UnsyncFile.h
+++ b/Engine/Source/Programs/Unsync/Private/UnsyncFile.h
@@ -453,6 +453,8 @@ bool		   FileCopy(const FPath& From, const FPath& To, std::error_code& OutErrorC
 bool		   FileCopyOverwrite(const FPath& From, const FPath& To, std::error_code& OutErrorCode);
 bool		   FileRemove(const FPath& Path, std::error_code& OutErrorCode);
 
+std::filesystem::recursive_directory_iterator RecursiveDirectoryScan(const FPath& Path);
+
 uint64 ToWindowsFileTime(const std::filesystem::file_time_type& T);
 
 struct FSyncFilter;
