unsync - Change the target name to use all lower-case letters

diff --git a/Engine/Source/Programs/Unsync/Unsync.Target.cs b/Engine/Source/Programs/Unsync/Unsync.Target.cs
--- a/Engine/Source/Programs/Unsync/Unsync.Target.cs
+++ b/Engine/Source/Programs/Unsync/Unsync.Target.cs
@@ -7,6 +7,7 @@ public class UnsyncTarget : TargetRules
 {
 	public UnsyncTarget(TargetInfo Target) : base(Target)
 	{
+		Name = "unsync";
 		Type = TargetType.Program;
 		LinkType = TargetLinkType.Monolithic;
 		LaunchModuleName = "Unsync";		
