unsync - Check directory manifest entry block size and mtime fields and report errors if invalid entries are found

diff --git a/Engine/Source/Programs/Unsync/Private/UnsyncCore.cpp b/Engine/Source/Programs/Unsync/Private/UnsyncCore.cpp
--- a/Engine/Source/Programs/Unsync/Private/UnsyncCore.cpp
+++ b/Engine/Source/Programs/Unsync/Private/UnsyncCore.cpp
@@ -2746,7 +2746,7 @@ SyncDirectory(const FSyncDirectoryOptions& SyncOptions)
 		if (!LoadDirectoryManifest(SourceDirectoryManifest, SourcePath, SourceManifestTempPath))
 		{
 			bSourceManifestOk = false;
-			UNSYNC_ERROR(L"Failed to read source directory manifest '%ls'", SourceManifestPath.wstring().c_str());
+			UNSYNC_ERROR(L"Failed to load source directory manifest '%ls'", SourceManifestPath.wstring().c_str());
 
 			return false;
 		}
diff --git a/Engine/Source/Programs/Unsync/Private/UnsyncSerialization.cpp b/Engine/Source/Programs/Unsync/Private/UnsyncSerialization.cpp
--- a/Engine/Source/Programs/Unsync/Private/UnsyncSerialization.cpp
+++ b/Engine/Source/Programs/Unsync/Private/UnsyncSerialization.cpp
@@ -466,6 +466,18 @@ LoadDirectoryManifest(FDirectoryManifest& OutManifest, const FPath& Root, FIORea
 
 			std::wstring Filename = ConvertUtf8ToWide(FilenameUtf8);
 
+			if (FileManifest.Mtime == 0)
+			{
+				UNSYNC_ERROR(L"Invalid manifest entry for file '%ls' (Mtime is 0)", Filename.c_str());
+				return false;
+			}
+
+			if (FileManifest.BlockSize == 0)
+			{
+				UNSYNC_ERROR(L"Invalid manifest entry for file '%ls' (BlockSize is 0)", Filename.c_str());
+				return false;
+			}
+
 			FileManifest.CurrentPath	= Root / Filename;
 			OutManifest.Files[Filename] = std::move(FileManifest);
 		}
