unsync - Suspend and resume sync task dispatcher using condition variable when scheduling a new task is not possible instead of using a poll loop

Also fix the background task memory usage estimation.

diff --git a/Engine/Source/Programs/Unsync/Private/UnsyncCore.cpp b/Engine/Source/Programs/Unsync/Private/UnsyncCore.cpp
--- a/Engine/Source/Programs/Unsync/Private/UnsyncCore.cpp
+++ b/Engine/Source/Programs/Unsync/Private/UnsyncCore.cpp
@@ -17,6 +17,7 @@
 #include <string>
 #include <unordered_map>
 #include <unordered_set>
+#include <condition_variable>
 
 UNSYNC_THIRD_PARTY_INCLUDES_START
 #include <blake3.h>
@@ -3236,6 +3237,9 @@ SyncDirectory(const FSyncDirectoryOptions& SyncOptions)
 		std::atomic<uint64> BackgroundTaskMemory	   = {};
 		std::atomic<uint64> RemainingSourceBytes	   = EstimatedNeedBytesFromSource;
 
+		std::mutex SchedulerMutex;
+		std::condition_variable SchedulerEvent;
+
 		while (!SyncTaskList.empty())
 		{
 			if (NumForegroundTasks == 0)
@@ -3247,41 +3251,42 @@ SyncDirectory(const FSyncDirectoryOptions& SyncOptions)
 				RemainingSourceBytes -= LocalTask.NeedBytesFromSource;
 
 				ForegroundTaskGroup.run(
-					[Task = std::move(LocalTask), &NumForegroundTasks, &SyncTaskBody, LogVerbose = GLogVerbose]() {
+					[Task = std::move(LocalTask), &SchedulerEvent, &NumForegroundTasks, &SyncTaskBody, LogVerbose = GLogVerbose]() {
 						FLogVerbosityScope VerbosityScope(LogVerbose);
 						SyncTaskBody(Task, false);
 						--NumForegroundTasks;
+						SchedulerEvent.notify_one();
 					});
 				continue;
 			}
 
 			const uint32 MaxBackgroundTasks = std::min<uint32>(8, GMaxThreads - 1);
 
-			uint64 CurrentBackgroundMemory = BackgroundTaskMemory;
-
 			if (NumBackgroundTasks < MaxBackgroundTasks && (SyncTaskList.front().NeedBytesFromSource < RemainingSourceBytes / 4) &&
-				(CurrentBackgroundMemory + SyncTaskList.front().TotalSizeBytes < BackgroundTaskMemoryBudget))
+				(BackgroundTaskMemory + SyncTaskList.front().TotalSizeBytes < BackgroundTaskMemoryBudget))
 			{
 				FileSyncTask LocalTask = std::move(SyncTaskList.front());
 				SyncTaskList.pop_front();
 
-				CurrentBackgroundMemory += LocalTask.TotalSizeBytes;
+				BackgroundTaskMemory += LocalTask.TotalSizeBytes;
 				++NumBackgroundTasks;
 
 				RemainingSourceBytes -= LocalTask.NeedBytesFromSource;
 
 				BackgroundTaskGroup.run(
-					[Task = std::move(LocalTask), &NumBackgroundTasks, &SyncTaskBody, &CurrentBackgroundMemory]() {
+					[Task = std::move(LocalTask), &SchedulerEvent, &NumBackgroundTasks, &SyncTaskBody, &BackgroundTaskMemory]() {
 						FLogVerbosityScope VerbosityScope(false);  // turn off logging from background threads
 						SyncTaskBody(Task, true);
 						--NumBackgroundTasks;
-						CurrentBackgroundMemory -= Task.TotalSizeBytes;
+						BackgroundTaskMemory -= Task.TotalSizeBytes;
+						SchedulerEvent.notify_one();
 					});
 
 				continue;
 			}
 
-			SchedulerSleep(10);
+			std::unique_lock<std::mutex> SchedulerLock(SchedulerMutex);
+			SchedulerEvent.wait(SchedulerLock);
 		}
 
 		ForegroundTaskGroup.wait();
