unsync - Enable quick sync mode by default, reduce log verbosity, fix file attribute cache lookup, bump version to 1.0.40

diff --git a/Engine/Source/Programs/Unsync/Private/UnsyncCmdSync.cpp b/Engine/Source/Programs/Unsync/Private/UnsyncCmdSync.cpp
--- a/Engine/Source/Programs/Unsync/Private/UnsyncCmdSync.cpp
+++ b/Engine/Source/Programs/Unsync/Private/UnsyncCmdSync.cpp
@@ -13,7 +13,10 @@ CmdSync(const FCmdSyncOptions& Options)
 	FPath			ResolvedSource = Options.Filter ? Options.Filter->Resolve(Options.Source) : Options.Source;
 
 	UNSYNC_VERBOSE(L"Sync source: '%ls'", Options.Source.wstring().c_str());
-	UNSYNC_VERBOSE(L"-- resolved: '%ls'", ResolvedSource.wstring().c_str());
+	if (Options.Source != ResolvedSource)
+	{
+		UNSYNC_VERBOSE(L"-- resolved: '%ls'", ResolvedSource.wstring().c_str());
+	}
 	UNSYNC_VERBOSE(L"Sync target: '%ls'", Options.Target.wstring().c_str());
 
 	const bool bSourcePathExists	 = PathExists(ResolvedSource, ErrorCode);
@@ -26,7 +29,7 @@ CmdSync(const FCmdSyncOptions& Options)
 	{
 		UNSYNC_VERBOSE(L"Manifest override: %ls", Options.SourceManifestOverride.wstring().c_str());
 
-		if (Options.Remote.IsValid() && Options.bQuickSourceValidation)
+		if (Options.Remote.IsValid() && !Options.bFullSourceScan)
 		{
 			bSourceFileSystemRequired = false;
 		}
@@ -65,8 +68,8 @@ CmdSync(const FCmdSyncOptions& Options)
 			SyncOptions.Remote				   = &Options.Remote;
 			SyncOptions.SyncFilter			   = Options.Filter;
 			SyncOptions.bCleanup			   = Options.bCleanup;
-			SyncOptions.bValidateSourceFiles   = !Options.bQuickSourceValidation;
-			SyncOptions.bQuickDifference	   = Options.bQuickDifference;
+			SyncOptions.bValidateSourceFiles   = Options.bFullSourceScan;
+			SyncOptions.bFullDifference		   = Options.bFullDifference;
 			SyncOptions.bValidateTargetFiles   = Options.bValidateTargetFiles;
 
 			return SyncDirectory(SyncOptions) ? 0 : 1;
diff --git a/Engine/Source/Programs/Unsync/Private/UnsyncCmdSync.h b/Engine/Source/Programs/Unsync/Private/UnsyncCmdSync.h
--- a/Engine/Source/Programs/Unsync/Private/UnsyncCmdSync.h
+++ b/Engine/Source/Programs/Unsync/Private/UnsyncCmdSync.h
@@ -18,9 +18,9 @@ struct FCmdSyncOptions
 
 	FRemoteDesc Remote;
 
-	bool bQuickDifference		= false;
-	bool bQuickSourceValidation = false;
-	bool bCleanup				= false;
+	bool bFullDifference = false;
+	bool bFullSourceScan = false;
+	bool bCleanup		 = false;
 
 	bool bValidateTargetFiles = true;  // WARNING: turning this off is intended only for testing/profiling
 
diff --git a/Engine/Source/Programs/Unsync/Private/UnsyncCore.cpp b/Engine/Source/Programs/Unsync/Private/UnsyncCore.cpp
--- a/Engine/Source/Programs/Unsync/Private/UnsyncCore.cpp
+++ b/Engine/Source/Programs/Unsync/Private/UnsyncCore.cpp
@@ -25,7 +25,7 @@ UNSYNC_THIRD_PARTY_INCLUDES_START
 #include <flat_hash_map.hpp>
 UNSYNC_THIRD_PARTY_INCLUDES_END
 
-#define UNSYNC_VERSION_STR "1.0.39"
+#define UNSYNC_VERSION_STR "1.0.40"
 
 namespace unsync {
 
@@ -2646,6 +2646,8 @@ ToPath(const std::wstring_view& Str)
 bool  // TODO: return a TResult
 SyncDirectory(const FSyncDirectoryOptions& SyncOptions)
 {
+	FTimePoint TimeBegin = TimePointNow();
+
 	const bool bFileSystemSource = SyncOptions.SourceType == ESyncSourceType::FileSystem;
 	const bool bServerSource	 = SyncOptions.SourceType == ESyncSourceType::Server;
 
@@ -2846,7 +2848,7 @@ SyncDirectory(const FSyncDirectoryOptions& SyncOptions)
 	bool			   bBaseDirectoryManifestValid = false;
 	bool			   bQuickDifferencePossible	   = false;
 
-	if (SyncOptions.bQuickDifference && SourceDirectoryManifest.Options.ChunkingAlgorithmId == EChunkingAlgorithmID::VariableBlocks
+	if (!SyncOptions.bFullDifference && SourceDirectoryManifest.Options.ChunkingAlgorithmId == EChunkingAlgorithmID::VariableBlocks
 		&& PathExists(BaseManifestPath))
 	{
 		bBaseDirectoryManifestValid = LoadDirectoryManifest(BaseDirectoryManifest, BasePath, BaseManifestPath);
@@ -2861,7 +2863,7 @@ SyncDirectory(const FSyncDirectoryOptions& SyncOptions)
 
 	if (bQuickDifferencePossible)
 	{
-		UNSYNC_VERBOSE(L"Quick file difference is allowed");
+		UNSYNC_VERBOSE(L"Quick file difference is allowed (use --full-diff option to override)");
 	}
 
 	for (const auto& SourceManifestIt : SourceDirectoryManifest.Files)
@@ -2934,12 +2936,12 @@ SyncDirectory(const FSyncDirectoryOptions& SyncOptions)
 
 			if (!BaseFileAttrib.bValid)
 			{
-				UNSYNC_VERBOSE(L"Dirty file: '%ls' (no base data)", SourceFilename.c_str());
+				UNSYNC_VERBOSE2(L"Dirty file: '%ls' (no base data)", SourceFilename.c_str());
 				StatFullCopy++;
 			}
 			else
 			{
-				UNSYNC_VERBOSE(L"Dirty file: '%ls'", SourceManifestIt.first.c_str());
+				UNSYNC_VERBOSE2(L"Dirty file: '%ls'", SourceManifestIt.first.c_str());
 				StatPartialCopy++;
 
 				if (bFileSystemSource && SyncOptions.bValidateSourceFiles && !SourceAttribCache.Exists(ResolvedSourceFilePath) &&
@@ -3300,13 +3302,17 @@ SyncDirectory(const FSyncDirectoryOptions& SyncOptions)
 		UNSYNC_ASSERT(RemainingSourceBytes == 0);
 
 		bool bAllBackgroundTasksSucceeded = true;
+		uint32 NumBackgroundSyncFiles = 0;
+		uint64 DownloadedBackgroundBytes = 0;
 		for (const BackgroundTaskResult& Item : BackgroundTaskResults)
 		{
 			if (Item.SyncResult.Succeeded())
 			{
-				UNSYNC_VERBOSE(L"Copied '%ls' (%ls, background)",
+				UNSYNC_VERBOSE2(L"Copied '%ls' (%ls, background)",
 							   Item.TargetFilePath.wstring().c_str(),
 							   Item.bIsPartialCopy ? L"partial" : L"full");
+				++NumBackgroundSyncFiles;
+				DownloadedBackgroundBytes += Item.SyncResult.SourceBytes;
 			}
 			else
 			{
@@ -3314,6 +3320,11 @@ SyncDirectory(const FSyncDirectoryOptions& SyncOptions)
 			}
 		}
 
+		if (NumBackgroundSyncFiles)
+		{
+			UNSYNC_VERBOSE(L"Background file copies: %d (%.2f MB)", NumBackgroundSyncFiles, SizeMb(DownloadedBackgroundBytes));
+		}
+
 		if (!bAllBackgroundTasksSucceeded)
 		{
 			for (const BackgroundTaskResult& Item : BackgroundTaskResults)
@@ -3360,6 +3371,9 @@ SyncDirectory(const FSyncDirectoryOptions& SyncOptions)
 	UNSYNC_VERBOSE(L"Copied from source: %.2f MB, copied from base: %.2f MB", SizeMb(StatSourceBytes), SizeMb(StatBaseBytes));
 	UNSYNC_VERBOSE(L"Sync completed %ls", bSyncSucceeded ? L"successfully" : L"with errors (see log for details)");
 
+	double ElapsedSeconds = DurationSec(TimeBegin, TimePointNow());
+	UNSYNC_VERBOSE2(L"Sync time: %.2f seconds", ElapsedSeconds);
+
 	return bSyncSucceeded;
 }
 
diff --git a/Engine/Source/Programs/Unsync/Private/UnsyncCore.h b/Engine/Source/Programs/Unsync/Private/UnsyncCore.h
--- a/Engine/Source/Programs/Unsync/Private/UnsyncCore.h
+++ b/Engine/Source/Programs/Unsync/Private/UnsyncCore.h
@@ -346,8 +346,7 @@ struct FSyncDirectoryOptions
 	bool			   bCleanup	  = false;	// whether to cleanup any files in the target directory that are not in the source manifest file
 	bool			   bValidateSourceFiles = true;	 // whether to check that all source files declared in the manifest are present/valid
 	bool			   bValidateTargetFiles = true;	 // WARNING: turning this off is intended only for testing/profiling
-	bool bQuickDifference = false;	// whether file difference could be computed quickly based only on source and base manifests (for
-									// variable chunking mode only)
+	bool			   bFullDifference = true;	// whether to run full file difference algorithm, even when there is an existing manifest
 };
 
 bool SyncDirectory(const FSyncDirectoryOptions& SyncOptions);
diff --git a/Engine/Source/Programs/Unsync/Private/UnsyncFile.cpp b/Engine/Source/Programs/Unsync/Private/UnsyncFile.cpp
--- a/Engine/Source/Programs/Unsync/Private/UnsyncFile.cpp
+++ b/Engine/Source/Programs/Unsync/Private/UnsyncFile.cpp
@@ -548,9 +548,11 @@ GetFileAttrib(const FPath& Path, FFileAttributeCache* AttribCache)
 {
 	FileAttributes Result;
 
+	FPath ExtendedPath = MakeExtendedAbsolutePath(Path);
+
 	if (AttribCache)
 	{
-		auto It = AttribCache->Map.find(Path);
+		auto It = AttribCache->Map.find(ExtendedPath);
 		if (It != AttribCache->Map.end())
 		{
 			Result = It->second;
@@ -558,8 +560,6 @@ GetFileAttrib(const FPath& Path, FFileAttributeCache* AttribCache)
 		}
 	}
 
-	FPath ExtendedPath = MakeExtendedAbsolutePath(Path);
-
 	WIN32_FILE_ATTRIBUTE_DATA AttributeData;
 	BOOL					  Ok = GetFileAttributesExW(ExtendedPath.c_str(), GetFileExInfoStandard, &AttributeData);
 	if (Ok)
diff --git a/Engine/Source/Programs/Unsync/Private/UnsyncMain.cpp b/Engine/Source/Programs/Unsync/Private/UnsyncMain.cpp
--- a/Engine/Source/Programs/Unsync/Private/UnsyncMain.cpp
+++ b/Engine/Source/Programs/Unsync/Private/UnsyncMain.cpp
@@ -68,16 +68,24 @@ InnerMain(int Argc, char** Argv)
 	bool					 bAllowInsecureTls		= false;
 	bool					 bUseTls				= false;
 	bool					 bUseDebugMode			= false;
-	bool					 bQuickSyncMode			= false;
-	bool					 bQuickDifference		= false;
-	bool					 bQuickSourceValidation = false;
 	bool					 bIncrementalMode		= false;
 	bool					 bNoOutputValidation	= false;
 	bool					 bNoCleanupAfterSync	= false;
+	bool					 bFullSourceScan		= false;
+	bool					 bFullDifference		= false;
 	int32					 CompressionLevel		= 3;
 	uint32					 DiffBlockSize			= uint32(4_KB);
 	uint32					 HashOrSyncBlockSize	= uint32(64_KB);
 
+	struct FDeprecatedOptions
+	{
+		bool bQuickSyncMode = false;
+		bool bQuickDifference = false;
+		bool bQuickSourceValidation = false;
+	} DeprecatedOptions;
+
+	const std::string HiddenGroupId; // CLI11 uses an empty string group name to mark arguments that should be hidden
+
 #if UNSYNC_USE_TLS
 	auto AddTlsOptions = [&CacertFilenameUtf8, &bUseTls, &bAllowInsecureTls](CLI::App* App)
 	{
@@ -161,15 +169,40 @@ InnerMain(int Argc, char** Argv)
 						HttpHeaderFilenameUtf8,
 						"Text file that contains any extra HTTP headers to pass to the remote server (auth tokens, etc.)");
 	SubSync->add_flag("--no-cleanup", bNoCleanupAfterSync, "Do not delete local files that aren't in the manifest after a successful sync");
-	SubSync->add_flag("--quick",
-					  bQuickSyncMode,
-					  "Quick sync mode that skips some of the validation steps (enables all '--quick-*' options)");
-	SubSync->add_flag("--quick-source-validation",
-					  bQuickSourceValidation,
-					  "Skip checking if all source files are present before starting a sync");
-	SubSync->add_flag("--quick-difference",
-					  bQuickDifference,
-					  "Allow computing file difference based on previous sync manifest and file timestamps");
+
+	// Deprecated --quick flag
+	SubSync
+		->add_flag("--quick",
+				   DeprecatedOptions.bQuickSyncMode,
+				   "Quick sync mode that skips some of the validation steps (enables all '--quick-*' options)")
+		->group(HiddenGroupId);
+
+	// Deprecated --quick-source-validation flag
+	SubSync
+		->add_flag("--quick-source-validation",
+				   DeprecatedOptions.bQuickSourceValidation,
+				   "Skip checking if all source files are present before starting a sync")
+		->group(HiddenGroupId);
+
+	// Deprecatred --quick-difference flag
+	SubSync
+		->add_flag("--quick-difference",
+				   DeprecatedOptions.bQuickDifference,
+				   "Allow computing file difference based on previous sync manifest and file timestamps")
+		->group(HiddenGroupId);
+
+	SubSync->add_flag("--full-diff",
+					  bFullDifference,
+					  "Run the full binary differencing algorithm on local files, even if there there is a compatible local directory "
+					  "manifest and file timestamps/sizes match. This is an extra precaution that will handle any unexpected local file "
+					  "modifications, but it should not be needed in a common case.");
+
+	SubSync->add_flag("--full-source-scan",
+					  bFullSourceScan,
+					  "Perform a full scan of the source directory to check that all files are present and their timestamps/sizes match "
+					  "the manifest. This is an extra precaution that will detect any missing or invalid remote files before running the "
+					  "sync process, however this can be very slow when dealing with large numbers of files and directories.");
+
 	SubSync->add_flag("--no-output-validation", bNoOutputValidation, "Skip final patched file block hash validation (DANGEROUS)");
 	SubSync->add_option("-b, --block", HashOrSyncBlockSize, "Block size in bytes (default=64KB)");
 	SubCommands.push_back(SubSync);
@@ -198,6 +231,7 @@ InnerMain(int Argc, char** Argv)
 	{
 		Subcommand->add_flag("-d, --dry, --dry-run", GDryRun, "Don't write any outputs to disk");
 		Subcommand->add_flag("-v, --verbose", GLogVerbose, "Verbose logging");
+		Subcommand->add_flag("--very-verbose", GLogVeryVerbose, "Very verbose logging");
 		Subcommand->add_flag("--progress", GLogProgress, "Output @progress and @status markers");
 		Subcommand->add_option("--threads", GMaxThreads, "Limit worker threads to specified number");
 		Subcommand->add_flag("--buffered-files", GForceBufferedFiles, "Always use buffered file IO");
@@ -219,10 +253,31 @@ InnerMain(int Argc, char** Argv)
 	_setmode(_fileno(stdout), _O_U8TEXT);
 #endif	// UNSYNC_PLATFORM_WINDOWS
 
-	if (bQuickSyncMode)
+	if (GLogVeryVerbose)
+	{
+		// Force verbose mode if very-verbose flag is present
+		GLogVerbose = true;
+	}
+
+	if (DeprecatedOptions.bQuickSyncMode)
+	{
+		UNSYNC_WARNING(
+			L"Quick mode is now the default and --quick flag is deprecated. Use --full-source-scan and --full-diff options to enable legacy "
+			L"default behavior.");
+	}
+
+	if (DeprecatedOptions.bQuickSourceValidation)
+	{
+		UNSYNC_WARNING(
+			L"Quick mode is now the default and --quick-source-validation flag is deprecated. Use --full-source-scan to enable legacy behavior "
+			L"that scans source directory.");
+	}
+
+	if (DeprecatedOptions.bQuickDifference)
 	{
-		bQuickSourceValidation = true;
-		bQuickDifference	   = true;
+		UNSYNC_WARNING(
+			L"Quick mode is now the default and --quick-difference flag is deprecated. Use --full-diff to enable legacy behavior performs "
+			L"full binary difference of local files even if timestamps and sizes match.");
 	}
 
 	if (bUseDebugMode)
@@ -331,6 +386,15 @@ InnerMain(int Argc, char** Argv)
 
 	UNSYNC_VERBOSE(L"UNSYNC %hs", GetVersionString().c_str());
 
+	if (GLogVeryVerbose)
+	{
+		UNSYNC_VERBOSE(L"Very verbose logging is enabled");
+	}
+	else if (GLogVerbose)
+	{
+		UNSYNC_VERBOSE(L"Verbose logging is enabled");
+	}
+
 	if (GDryRun)
 	{
 		UNSYNC_VERBOSE(L">>> DRY RUN <<<");
@@ -546,8 +610,8 @@ InnerMain(int Argc, char** Argv)
 		SyncOptions.Target				   = TargetFilename;
 		SyncOptions.SourceManifestOverride = SourceManifestFilename;
 		SyncOptions.Remote				   = RemoteDesc;
-		SyncOptions.bQuickDifference	   = bQuickDifference;
-		SyncOptions.bQuickSourceValidation = bQuickSourceValidation;
+		SyncOptions.bFullDifference		   = bFullDifference;
+		SyncOptions.bFullSourceScan		   = bFullSourceScan;
 		SyncOptions.bCleanup			   = !bNoCleanupAfterSync;
 		SyncOptions.BlockSize			   = HashOrSyncBlockSize;
 		SyncOptions.Filter				   = &SyncFilter;
