sync - Change `query mirrors` command output to use integer milliseconds formatting for ping times, bump version to 1.0.41

diff --git a/Engine/Source/Programs/Unsync/Private/UnsyncCmdQuery.cpp b/Engine/Source/Programs/Unsync/Private/UnsyncCmdQuery.cpp
--- a/Engine/Source/Programs/Unsync/Private/UnsyncCmdQuery.cpp
+++ b/Engine/Source/Programs/Unsync/Private/UnsyncCmdQuery.cpp
@@ -143,11 +143,11 @@ CmdQuery(const FCmdQueryOptions& Options)
 			const FMirrorInfo& Mirror = Mirrors[I];
 
 			LogPrintf(ELogLevel::Info,
-					  L"  {\"address\":\"%hs\", \"port\":%d, \"ok\":%hs, \"ping\":\"%.4f\", \"name\":\"%hs\"}%hs\n",
+					  L"  {\"address\":\"%hs\", \"port\":%d, \"ok\":%hs, \"ping\":%d, \"name\":\"%hs\"}%hs\n",
 					  Mirror.Address.c_str(),
 					  Mirror.Port,
 					  Mirror.Ping > 0 ? "true" : "false",
-					  Mirror.Ping * 1000.0,
+					  int32(Mirror.Ping * 1000.0),
 					  Mirror.Name.c_str(),
 					  I + 1 == Mirrors.size() ? "" : ",");
 		}
diff --git a/Engine/Source/Programs/Unsync/Private/UnsyncCore.cpp b/Engine/Source/Programs/Unsync/Private/UnsyncCore.cpp
--- a/Engine/Source/Programs/Unsync/Private/UnsyncCore.cpp
+++ b/Engine/Source/Programs/Unsync/Private/UnsyncCore.cpp
@@ -25,7 +25,7 @@ UNSYNC_THIRD_PARTY_INCLUDES_START
 #include <flat_hash_map.hpp>
 UNSYNC_THIRD_PARTY_INCLUDES_END
 
-#define UNSYNC_VERSION_STR "1.0.40"
+#define UNSYNC_VERSION_STR "1.0.41"
 
 namespace unsync {
 
diff --git a/Engine/Source/Programs/Unsync/Private/UnsyncMain.cpp b/Engine/Source/Programs/Unsync/Private/UnsyncMain.cpp
--- a/Engine/Source/Programs/Unsync/Private/UnsyncMain.cpp
+++ b/Engine/Source/Programs/Unsync/Private/UnsyncMain.cpp
@@ -193,7 +193,7 @@ InnerMain(int Argc, char** Argv)
 
 	SubSync->add_flag("--full-diff",
 					  bFullDifference,
-					  "Run the full binary differencing algorithm on local files, even if there there is a compatible local directory "
+					  "Run the full binary differencing algorithm on local files, even if there is a compatible local directory "
 					  "manifest and file timestamps/sizes match. This is an extra precaution that will handle any unexpected local file "
 					  "modifications, but it should not be needed in a common case.");
 
