unsync - Handle local file name case differences during sync, bump version to 1.0.42

When running on case-insensitive file systems, local files and directories will be renamed to match the case specified in remote directory manifest.

diff --git a/Engine/Source/Programs/Unsync/Private/UnsyncCore.cpp b/Engine/Source/Programs/Unsync/Private/UnsyncCore.cpp
--- a/Engine/Source/Programs/Unsync/Private/UnsyncCore.cpp
+++ b/Engine/Source/Programs/Unsync/Private/UnsyncCore.cpp
@@ -25,7 +25,7 @@ UNSYNC_THIRD_PARTY_INCLUDES_START
 #include <flat_hash_map.hpp>
 UNSYNC_THIRD_PARTY_INCLUDES_END
 
-#define UNSYNC_VERSION_STR "1.0.41"
+#define UNSYNC_VERSION_STR "1.0.42"
 
 namespace unsync {
 
@@ -2600,6 +2600,176 @@ CopyFileIfNewer(const FPath& Source, const FPath& Target)
 	return Ec;
 }
 
+static bool
+IsNonCaseSensitiveFileSystem(const FPath& ExistingPath)
+{
+	UNSYNC_ASSERTF(PathExists(ExistingPath), L"IsCaseSensitiveFileSystem must be called with a path that exists on disk");
+
+	// Assume file system is case-sensitive if all-upper and all-lower versions of the path exist and resolve to the same FS entry.
+	// This is not 100% robust due to symlinks, but is good enough for most practical purposes.
+
+	FPath PathUpper = StringToUpper(ExistingPath.wstring());
+	FPath PathLower = StringToLower(ExistingPath.wstring());
+
+	if (PathExists(PathUpper) && PathExists(PathLower))
+	{
+		return std::filesystem::equivalent(ExistingPath, PathUpper) && std::filesystem::equivalent(PathLower, PathUpper);
+	}
+	else
+	{
+		return false;
+	}
+}
+
+struct FPendingFileRename
+{
+	std::wstring Old;
+	std::wstring New;
+};
+
+// Updates the target directory manifest filename case to be consistent with reference.
+// Internally we always perform case-sensitive path comparisons, however on non-case-sensitive filesystems some local files may be renamed to a mismatching case.
+// We can update the locally-generated manifest to take the case from the reference manifest for equivalent paths.
+// Returns a list of files that should be renamed on disk.
+static std::vector<FPendingFileRename>
+FixManifestFileNameCases(
+	FDirectoryManifest& TargetDirectoryManifest,
+	const FDirectoryManifest& ReferenceManifest)
+{
+	// Build a lookup table of lowercase -> original file names and detect potential case conflicts (which will explode on Windows and Mac)
+
+	std::unordered_map<std::wstring, std::wstring> ReferenceFileNamesLowerCase;
+	bool bFoundCaseConflicts = false;
+	for (auto& ReferenceManifestEntry : ReferenceManifest.Files)
+	{
+		std::wstring FileNameLowerCase = StringToLower(ReferenceManifestEntry.first);
+		auto InsertResult = ReferenceFileNamesLowerCase.insert(std::pair<std::wstring, std::wstring>(FileNameLowerCase, ReferenceManifestEntry.first));
+
+		if (!InsertResult.second)
+		{
+			UNSYNC_WARNING(L"Found file name case conflict: '%ls'", ReferenceManifestEntry.first.c_str());
+			bFoundCaseConflicts = true;
+		}
+	}
+
+	if (bFoundCaseConflicts)
+	{
+		UNSYNC_WARNING(L"File name case conflicts will result in issues on case-insensitive systems, such as Windows and macOS.");
+	}
+
+	// Find inconsistently-cased files and add them to a list to be fixed up
+
+	std::vector<FPendingFileRename> FixupEntries;
+
+	for (auto& TargetManifestEntry : TargetDirectoryManifest.Files)
+	{
+		const std::wstring& TargetFileName = TargetManifestEntry.first;
+		if (ReferenceManifest.Files.find(TargetFileName) == ReferenceManifest.Files.end())
+		{
+			std::wstring TargetFileNameLowerCase = StringToLower(TargetFileName);
+			auto ReferenceIt = ReferenceFileNamesLowerCase.find(TargetFileNameLowerCase);
+			if (ReferenceIt != ReferenceFileNamesLowerCase.end())
+			{
+				FixupEntries.push_back({TargetFileName, ReferenceIt->second});
+			}
+		}
+	}
+
+	// Re-add file manifests under the correct names
+
+	for (const FPendingFileRename& Entry : FixupEntries)
+	{
+		auto It = TargetDirectoryManifest.Files.find(Entry.Old);
+		UNSYNC_ASSERT(It != TargetDirectoryManifest.Files.end());
+
+		FFileManifest Manifest;
+		std::swap(It->second, Manifest);
+		TargetDirectoryManifest.Files.erase(Entry.Old);
+		TargetDirectoryManifest.Files.insert(std::pair(Entry.New, std::move(Manifest)));
+	}
+
+	return FixupEntries;
+}
+
+// Takes a list of file names that require case fixup and performs the necessary renaming.
+// Handles renaming of intermediate directories as well as the leaf files.
+// Quite wasteful in terms of mallocs, but doesn't matter since we're about to touch the file system anyway.
+static bool
+FixFileNameCases(const FPath& RootPath, const std::vector<FPendingFileRename>& PendingRenames)
+{
+	std::vector<FPendingFileRename> UniqueRenames;
+	std::unordered_set<std::wstring> UniqueRenamesSet;
+
+	// Build a rename schedule, with only unique entries (taking subdirectories into account)
+
+	for (const FPendingFileRename& Entry : PendingRenames)
+	{
+		UNSYNC_ASSERTF(StringToLower(Entry.Old) == StringToLower(Entry.New),
+			L"FixFileNameCases expects inputs that are different only by case. Old: '%ls', New: '%ls'",
+			Entry.Old.c_str(), Entry.New.c_str());
+
+		FPath OldPath = Entry.Old;
+		FPath NewPath = Entry.New;
+
+		auto ItOld = OldPath.begin();
+		auto ItNew = NewPath.begin();
+
+		FPath OldPathPart;
+		FPath NewPathPart;
+
+		while (ItOld != OldPath.end())
+		{
+			OldPathPart /= *ItOld;
+			NewPathPart /= *ItNew;
+
+			if (*ItOld != *ItNew)
+			{
+				auto InsertResult = UniqueRenamesSet.insert(OldPathPart);
+				if (InsertResult.second)
+				{
+					UniqueRenames.push_back({OldPathPart, NewPathPart});
+				}
+			}
+
+			++ItOld;
+			++ItNew;
+		}
+	}
+
+	std::sort(UniqueRenames.begin(), UniqueRenames.end(), [](const FPendingFileRename& A, const FPendingFileRename& B)
+	{
+		return A.Old < B.Old;
+	});
+
+	// Perform actual renaming
+
+	for (const FPendingFileRename& Entry : UniqueRenames)
+	{
+		FPath OldPath = RootPath / Entry.Old;
+		FPath NewPath = RootPath / Entry.New;
+
+		std::error_code ErrorCode;
+
+		if (GDryRun)
+		{
+			UNSYNC_VERBOSE(L"Renaming '%ls' -> '%ls' (skipped due to dry run mode)", Entry.Old.c_str(), Entry.New.c_str());
+		}
+		else
+		{
+			UNSYNC_VERBOSE(L"Renaming '%ls' -> '%ls'", Entry.Old.c_str(), Entry.New.c_str());
+			FileRename(OldPath, NewPath, ErrorCode);
+		}
+
+		if (ErrorCode)
+		{
+			UNSYNC_VERBOSE(L"Failed to rename file. System error code %d: %hs", ErrorCode.value(), ErrorCode.message().c_str());
+			return false;
+		}
+	}
+
+	return true;
+}
+
 // Delete files from target directory that are not in the source directory manifest
 static void
 DeleteUnnecessaryFiles(const FPath&				 TargetDirectory,
@@ -2800,6 +2970,21 @@ SyncDirectory(const FSyncDirectoryOptions& SyncOptions)
 	FAlgorithmOptions  Algorithm			   = SourceDirectoryManifest.Options;
 	FDirectoryManifest TargetDirectoryManifest = CreateDirectoryManifest(TargetPath, 0, Algorithm);
 
+	if (IsNonCaseSensitiveFileSystem(TargetTempPath))
+	{
+		std::vector<FPendingFileRename> PendingRenames;
+		PendingRenames = FixManifestFileNameCases(TargetDirectoryManifest, SourceDirectoryManifest);
+		if (!PendingRenames.empty())
+		{
+			UNSYNC_VERBOSE(L"Fixing inconsistent case of target files");
+			UNSYNC_LOG_INDENT;
+			if (!FixFileNameCases(TargetPath, PendingRenames))
+			{
+				return false;
+			}
+		}
+	}
+
 	uint32 StatSkipped	   = 0;
 	uint32 StatFullCopy	   = 0;
 	uint32 StatPartialCopy = 0;
diff --git a/Engine/Source/Programs/Unsync/Private/UnsyncUtil.cpp b/Engine/Source/Programs/Unsync/Private/UnsyncUtil.cpp
--- a/Engine/Source/Programs/Unsync/Private/UnsyncUtil.cpp
+++ b/Engine/Source/Programs/Unsync/Private/UnsyncUtil.cpp
@@ -117,6 +117,14 @@ StringToLower(const std::wstring& Input)
 	return Result;
 }
 
+std::wstring
+StringToUpper(const std::wstring& Input)
+{
+	std::wstring Result = Input;
+	std::transform(Result.begin(), Result.end(), Result.begin(), [](int32 C) { return wchar_t(::toupper(C)); });
+	return Result;
+}
+
 FDfsMirrorInfo
 DfsEnumerate(const FPath& Root)
 {
diff --git a/Engine/Source/Programs/Unsync/Private/UnsyncUtil.h b/Engine/Source/Programs/Unsync/Private/UnsyncUtil.h
--- a/Engine/Source/Programs/Unsync/Private/UnsyncUtil.h
+++ b/Engine/Source/Programs/Unsync/Private/UnsyncUtil.h
@@ -130,6 +130,7 @@ std::wstring ConvertUtf8ToWide(std::string_view StringUtf8);
 std::string	 ConvertWideToUtf8(std::wstring_view StringWide);
 
 std::wstring StringToLower(const std::wstring& Input);
+std::wstring StringToUpper(const std::wstring& Input);
 
 // Returns a list of alternative DFS paths for a given root
 struct FDfsStorageInfo
