unsync - Fix compile errors on Linux

diff --git a/Engine/Source/Programs/Unsync/Private/UnsyncCore.cpp b/Engine/Source/Programs/Unsync/Private/UnsyncCore.cpp
--- a/Engine/Source/Programs/Unsync/Private/UnsyncCore.cpp
+++ b/Engine/Source/Programs/Unsync/Private/UnsyncCore.cpp
@@ -2698,7 +2698,7 @@ static bool
 FixFileNameCases(const FPath& RootPath, const std::vector<FPendingFileRename>& PendingRenames)
 {
 	std::vector<FPendingFileRename> UniqueRenames;
-	std::unordered_set<std::wstring> UniqueRenamesSet;
+	std::unordered_set<FPath::string_type> UniqueRenamesSet;
 
 	// Build a rename schedule, with only unique entries (taking subdirectories into account)
 
@@ -2724,10 +2724,10 @@ FixFileNameCases(const FPath& RootPath, const std::vector<FPendingFileRename>& P
 
 			if (*ItOld != *ItNew)
 			{
-				auto InsertResult = UniqueRenamesSet.insert(OldPathPart);
+				auto InsertResult = UniqueRenamesSet.insert(OldPathPart.native());
 				if (InsertResult.second)
 				{
-					UniqueRenames.push_back({OldPathPart, NewPathPart});
+					UniqueRenames.push_back({OldPathPart.wstring(), NewPathPart.wstring()});
 				}
 			}
 
