unsync - Add a special case for utf8<->wide string conversion when all characters are ASCII

diff --git a/Engine/Source/Programs/Unsync/Private/UnsyncUtil.cpp b/Engine/Source/Programs/Unsync/Private/UnsyncUtil.cpp
--- a/Engine/Source/Programs/Unsync/Private/UnsyncUtil.cpp
+++ b/Engine/Source/Programs/Unsync/Private/UnsyncUtil.cpp
@@ -88,18 +88,62 @@ FTimingLogger::~FTimingLogger()
 	}
 }
 
+template<typename T>
+static bool
+IsTrivialAsciiString(const T& Input)
+{
+	for (auto c : Input)
+	{
+		if ((unsigned)c > 127)
+		{
+			return false;
+		}
+	}
+	return true;
+}
+
 std::wstring
 ConvertUtf8ToWide(std::string_view StringUtf8)
 {
-	std::wstring_convert<std::codecvt_utf8<wchar_t>> Cvt;
-	return Cvt.from_bytes(StringUtf8.data(), StringUtf8.data() + StringUtf8.length());
+	std::wstring Result;
+
+	if (IsTrivialAsciiString(StringUtf8))
+	{
+		Result.reserve(StringUtf8.length());
+		for (char c : StringUtf8)
+		{
+			Result.push_back((wchar_t)c);
+		}
+	}
+	else
+	{
+		std::wstring_convert<std::codecvt_utf8<wchar_t>> Cvt;
+		Result = Cvt.from_bytes(StringUtf8.data(), StringUtf8.data() + StringUtf8.length());
+	}
+
+	return Result;
 }
 
 std::string
 ConvertWideToUtf8(std::wstring_view StringWide)
 {
-	std::wstring_convert<std::codecvt_utf8<wchar_t>> Cvt;
-	return Cvt.to_bytes(StringWide.data(), StringWide.data() + StringWide.length());
+	std::string Result;
+
+	if (IsTrivialAsciiString(StringWide))
+	{
+		Result.reserve(StringWide.length());
+		for (wchar_t wc : StringWide)
+		{
+			Result.push_back((char)wc);
+		}
+	}
+	else
+	{
+		std::wstring_convert<std::codecvt_utf8<wchar_t>> Cvt;
+		Result = Cvt.to_bytes(StringWide.data(), StringWide.data() + StringWide.length());
+	}
+
+	return Result;
 }
 
 const bool
