unsync - NormalizeFilenameUtf8 now also makes path absolute

diff --git a/Engine/Source/Programs/Unsync/Private/UnsyncUtil.cpp b/Engine/Source/Programs/Unsync/Private/UnsyncUtil.cpp
--- a/Engine/Source/Programs/Unsync/Private/UnsyncUtil.cpp
+++ b/Engine/Source/Programs/Unsync/Private/UnsyncUtil.cpp
@@ -266,7 +266,7 @@ DfsEnumerate(const FPath& Root)
 }
 
 FPath
-NormalizeFilenameUtf8(std::string& InFilename)
+NormalizeFilenameUtf8(const std::string& InFilename)
 {
 	std::string_view Filename	   = InFilename;
 	std::string_view FileUrlPrefix = "file://";
@@ -276,7 +276,9 @@ NormalizeFilenameUtf8(std::string& InFilename)
 	}
 
 	FPath FilenameAsPath = ConvertUtf8ToWide(Filename);
-	return FilenameAsPath.lexically_normal();
+	FPath NormalPath = FilenameAsPath.lexically_normal();
+	FPath AbsoluteNormalPath = std::filesystem::absolute(NormalPath);
+	return AbsoluteNormalPath;
 }
 
 const FBuffer&
diff --git a/Engine/Source/Programs/Unsync/Private/UnsyncUtil.h b/Engine/Source/Programs/Unsync/Private/UnsyncUtil.h
--- a/Engine/Source/Programs/Unsync/Private/UnsyncUtil.h
+++ b/Engine/Source/Programs/Unsync/Private/UnsyncUtil.h
@@ -276,7 +276,7 @@ CountLeadingZeros64(uint64 X)
 	return ((63 - XLog2) & Mask) | (64 & ~Mask);
 }
 
-FPath NormalizeFilenameUtf8(std::string& InFilename);
+FPath NormalizeFilenameUtf8(const std::string& InFilename);
 
 const FBuffer& GetSystemRootCerts();
 
