unsync - Send basic sync stat telemetry to the server (if server supports it)

diff --git a/Engine/Source/Programs/Unsync/CMakeLists.txt b/Engine/Source/Programs/Unsync/CMakeLists.txt
--- a/Engine/Source/Programs/Unsync/CMakeLists.txt
+++ b/Engine/Source/Programs/Unsync/CMakeLists.txt
@@ -89,7 +89,7 @@ set(src
 
 # Add third party code
 
-set(blake3dir ${CMAKE_CURRENT_SOURCE_DIR}/../../ThirdParty/BLAKE3/0.3.7/c)
+set(blake3dir ${CMAKE_CURRENT_SOURCE_DIR}/../../ThirdParty/BLAKE3/1.3.1/c)
 
 set(thirdparty
 	${blake3dir}/blake3_dispatch.c
diff --git a/Engine/Source/Programs/Unsync/Private/UnsyncCmdQuery.cpp b/Engine/Source/Programs/Unsync/Private/UnsyncCmdQuery.cpp
--- a/Engine/Source/Programs/Unsync/Private/UnsyncCmdQuery.cpp
+++ b/Engine/Source/Programs/Unsync/Private/UnsyncCmdQuery.cpp
@@ -20,18 +20,6 @@ struct FMirrorInfo
 };
 using FMirrorInfoResult = TResult<std::vector<FMirrorInfo>>;
 
-static FHttpResponse
-HttpRequest(const FRemoteDesc& RemoteDesc, EHttpMethod Method, std::string_view RequestUrl)
-{
-	FTlsClientSettings TlsSettings = RemoteDesc.GetTlsClientSettings();
-	FHttpConnection	   Connection(RemoteDesc.HostAddress, RemoteDesc.HostPort, &TlsSettings);
-	FHttpRequest	   Request;
-	Request.CustomHeaders  = RemoteDesc.HttpHeaders;
-	Request.Url			   = RequestUrl;
-	FHttpResponse Response = HttpRequest(Connection, Request);
-	return Response;
-}
-
 // Runs a basic HTTP request against the remote server and returns the time it took to get the response, -1 if connection could not be
 // established.
 static double
diff --git a/Engine/Source/Programs/Unsync/Private/UnsyncCore.cpp b/Engine/Source/Programs/Unsync/Private/UnsyncCore.cpp
--- a/Engine/Source/Programs/Unsync/Private/UnsyncCore.cpp
+++ b/Engine/Source/Programs/Unsync/Private/UnsyncCore.cpp
@@ -3051,6 +3051,8 @@ SyncDirectory(const FSyncDirectoryOptions& SyncOptions)
 		UNSYNC_VERBOSE(L"Quick file difference is allowed (use --full-diff option to override)");
 	}
 
+	uint64 TotalSourceSize = 0;
+
 	for (const auto& SourceManifestIt : SourceDirectoryManifest.Files)
 	{
 		const std::wstring& SourceFilename = SourceManifestIt.first;
@@ -3064,6 +3066,8 @@ SyncDirectory(const FSyncDirectoryOptions& SyncOptions)
 
 		const FFileManifest& SourceFileManifest = SourceManifestIt.second;
 
+		TotalSourceSize += SourceFileManifest.Size;
+
 		auto TargetManifestIt = TargetDirectoryManifest.Files.find(SourceFilename);
 		if (TargetManifestIt != TargetDirectoryManifest.Files.end())
 		{
@@ -3328,7 +3332,7 @@ SyncDirectory(const FSyncDirectoryOptions& SyncOptions)
 			// TODO: bail out if remote connection is required for the download,
 			// such as when downloading data purely from Jupiter.
 
-			UNSYNC_VERBOSE(L"Attempting to sync without remote server connection");
+			UNSYNC_WARNING(L"Attempting to sync without remote server connection");
 		}
 	}
 
@@ -3559,6 +3563,24 @@ SyncDirectory(const FSyncDirectoryOptions& SyncOptions)
 	double ElapsedSeconds = DurationSec(TimeBegin, TimePointNow());
 	UNSYNC_VERBOSE2(L"Sync time: %.2f seconds", ElapsedSeconds);
 
+	if (ProxyPool.IsValid() && ProxyPool.GetFeatures().bTelemetry)
+	{
+		FTelemetryEventSyncComplete Event;
+
+		Event.Session		   = ProxyPool.GetSessionId();
+		Event.Source		   = ConvertWideToUtf8(SourcePath.wstring());
+		Event.TotalBytes	   = TotalSourceSize;
+		Event.SourceBytes	   = StatSourceBytes;
+		Event.BaseBytes		   = StatBaseBytes;
+		Event.SkippedFiles	   = StatSkipped;
+		Event.FullCopyFiles	   = StatFullCopy;
+		Event.PartialCopyFiles = StatPartialCopy;
+		Event.Elapsed		   = ElapsedSeconds;
+		Event.bSuccess		   = bSyncSucceeded;
+
+		ProxyPool.SendTelemetryEvent(Event);
+	}
+
 	return bSyncSucceeded;
 }
 
diff --git a/Engine/Source/Programs/Unsync/Private/UnsyncHttp.cpp b/Engine/Source/Programs/Unsync/Private/UnsyncHttp.cpp
--- a/Engine/Source/Programs/Unsync/Private/UnsyncHttp.cpp
+++ b/Engine/Source/Programs/Unsync/Private/UnsyncHttp.cpp
@@ -3,6 +3,7 @@
 #include "UnsyncHttp.h"
 #include "UnsyncCore.h"
 #include "UnsyncUtil.h"
+#include "UnsyncRemote.h"
 
 #include <http_parser.h>
 #include <string.h>
@@ -229,6 +230,35 @@ HttpRequest(FHttpConnection& Connection, const FHttpRequest& Request)
 	return Result;
 }
 
+FHttpResponse
+HttpRequest(const FRemoteDesc& RemoteDesc,
+			EHttpMethod		   Method,
+			std::string_view   RequestUrl,
+			EHttpContentType   PayloadContentType,
+			FBufferView		   Payload)
+{
+	FTlsClientSettings TlsSettings = RemoteDesc.GetTlsClientSettings();
+	FHttpConnection	   Connection(RemoteDesc.HostAddress, RemoteDesc.HostPort, &TlsSettings);
+
+	FHttpRequest Request;
+
+	Request.Method			   = Method;
+	Request.CustomHeaders	   = RemoteDesc.HttpHeaders;
+	Request.Url				   = RequestUrl;
+	Request.PayloadContentType = PayloadContentType;
+	Request.Payload			   = Payload;
+
+	FHttpResponse Response = HttpRequest(Connection, Request);
+
+	return Response;
+}
+
+FHttpResponse
+HttpRequest(const FRemoteDesc& RemoteDesc, EHttpMethod Method, std::string_view RequestUrl)
+{
+	return HttpRequest(RemoteDesc, Method, RequestUrl, EHttpContentType::Unknown, {});
+}
+
 bool
 HttpRequestBegin(FHttpConnection& Connection, const FHttpRequest& Request)
 {
@@ -297,6 +327,9 @@ HttpRequestBegin(FHttpConnection& Connection, const FHttpRequest& Request)
 			case EHttpContentType::Application_UECB:
 				HttpHeader += "Content-type: application/x-ue-cb\r\n";
 				break;
+			case EHttpContentType::Application_Json:
+				HttpHeader += "Content-type: application/json\r\n";
+				break;
 
 			default:
 				UNSYNC_FATAL(L"HTTP content type not supported");
diff --git a/Engine/Source/Programs/Unsync/Private/UnsyncHttp.h b/Engine/Source/Programs/Unsync/Private/UnsyncHttp.h
--- a/Engine/Source/Programs/Unsync/Private/UnsyncHttp.h
+++ b/Engine/Source/Programs/Unsync/Private/UnsyncHttp.h
@@ -13,6 +13,8 @@
 
 namespace unsync {
 
+struct FRemoteDesc;
+
 enum class EHttpContentType
 {
 	Unknown,
@@ -125,6 +127,12 @@ HttpRequest(FHttpConnection& Connection, EHttpMethod Method, std::string_view Ur
 	return HttpRequest(Connection, Request);
 }
 
+FHttpResponse
+HttpRequest(const FRemoteDesc& RemoteDesc, EHttpMethod Method, std::string_view RequestUrl);
+
+FHttpResponse
+HttpRequest(const FRemoteDesc& RemoteDesc, EHttpMethod Method, std::string_view RequestUrl, EHttpContentType PayloadContentType, FBufferView Payload);
+
 // Pipelined HTTP request API
 
 bool HttpRequestBegin(FHttpConnection& Connection, const FHttpRequest& Request);
diff --git a/Engine/Source/Programs/Unsync/Private/UnsyncProxy.cpp b/Engine/Source/Programs/Unsync/Private/UnsyncProxy.cpp
--- a/Engine/Source/Programs/Unsync/Private/UnsyncProxy.cpp
+++ b/Engine/Source/Programs/Unsync/Private/UnsyncProxy.cpp
@@ -5,11 +5,14 @@
 #include "UnsyncFile.h"
 #include "UnsyncJupiter.h"
 
+#include <json11.hpp>
+#include <fmt/format.h>
+
 namespace unsync {
 
 struct FUnsyncProtocolImpl : FRemoteProtocolBase
 {
-	FUnsyncProtocolImpl(const FRemoteDesc& InRemoteDesc, const FBlockRequestMap* InRequestMap, const FTlsClientSettings* TlsSettings);
+	FUnsyncProtocolImpl(const FRemoteDesc& InRemoteDesc, const FRemoteProtocolFeatures& InFeatures, const FBlockRequestMap* InRequestMap, const FTlsClientSettings* TlsSettings);
 	virtual ~FUnsyncProtocolImpl() override;
 	virtual bool			 IsValid() const override;
 	virtual FDownloadResult	 Download(const TArrayView<FNeedBlock> NeedBlocks, const FBlockDownloadCallback& CompletionCallback) override;
@@ -24,9 +27,23 @@ struct FUnsyncProtocolImpl : FRemoteProtocolBase
 
 	bool						 bIsConnetedToHost = false;
 	std::unique_ptr<FSocketBase> SocketHandle;
+
+	const FRemoteProtocolFeatures Features;
+
+	struct FHelloResponse
+	{
+		std::string Name;
+		std::string VersionNumber;
+		std::string VersionGit;
+		std::string SessionId;
+		std::vector<std::string> FeatureNames;
+		FRemoteProtocolFeatures Features;
+	};
+	static TResult<FHelloResponse> QueryHello(const FRemoteDesc& RemoteDesc);
+	static void SendTelemetryEvent(const FRemoteDesc& RemoteDesc, const FTelemetryEventSyncComplete& Event);
 };
 
-FProxy::FProxy(const FRemoteDesc& RemoteDesc, const FBlockRequestMap* InRequestMap)
+FProxy::FProxy(const FRemoteDesc& RemoteDesc, const FRemoteProtocolFeatures& InFeatures, const FBlockRequestMap* InRequestMap)
 {
 	UNSYNC_ASSERT(InRequestMap);
 
@@ -46,7 +63,7 @@ FProxy::FProxy(const FRemoteDesc& RemoteDesc, const FBlockRequestMap* InRequestM
 	}
 	else if (RemoteDesc.Protocol == EProtocolFlavor::Unsync)
 	{
-		ProtocolImpl = std::unique_ptr<FRemoteProtocolBase>(new FUnsyncProtocolImpl(RemoteDesc, InRequestMap, &TlsSettings));
+		ProtocolImpl = std::unique_ptr<FRemoteProtocolBase>(new FUnsyncProtocolImpl(RemoteDesc, InFeatures, InRequestMap, &TlsSettings));
 	}
 	else
 	{
@@ -54,10 +71,12 @@ FProxy::FProxy(const FRemoteDesc& RemoteDesc, const FBlockRequestMap* InRequestM
 	}
 }
 
-FUnsyncProtocolImpl::FUnsyncProtocolImpl(const FRemoteDesc&		   RemoteDesc,
-										 const FBlockRequestMap*   InRequestMap,
-										 const FTlsClientSettings* TlsSettings)
+FUnsyncProtocolImpl::FUnsyncProtocolImpl(const FRemoteDesc&				RemoteDesc,
+										 const FRemoteProtocolFeatures& InFeatures,
+										 const FBlockRequestMap*		InRequestMap,
+										 const FTlsClientSettings*		TlsSettings)
 : FRemoteProtocolBase(RemoteDesc, InRequestMap)
+, Features(InFeatures)
 {
 #if UNSYNC_USE_TLS
 	if (RemoteDesc.bTlsEnable && TlsSettings)
@@ -331,6 +350,115 @@ FUnsyncProtocolImpl::GetSocketSecurity() const
 	}
 }
 
+TResult<FUnsyncProtocolImpl::FHelloResponse> FUnsyncProtocolImpl::QueryHello(const FRemoteDesc& RemoteDesc)
+{
+	const char* Url = "/api/v1/hello";
+	FHttpResponse Response = HttpRequest(RemoteDesc, EHttpMethod::GET, Url);
+	if (!Response.Success())
+	{
+		UNSYNC_ERROR(L"Failed to establish connection to UNSYNC server. Error code: %d.", Response.Code);
+		return HttpError(Response.Code);
+	}
+
+	FHelloResponse Result;
+
+	using namespace json11;
+	std::string JsonString = std::string(Response.AsStringView());
+
+	std::string JsonErrorString;
+	Json		JsonObject = Json::parse(JsonString, JsonErrorString);
+
+	if (!JsonErrorString.empty())
+	{
+		return AppError(std::string("JSON parse error while connecting to UNSYNC server: ") + JsonErrorString);
+	}
+
+	if (auto& Field = JsonObject["service"]; Field.is_string())
+	{
+		const std::string& FieldVal = Field.string_value();
+		if (FieldVal != "unsync")
+		{
+			return AppError(fmt::format("Expected service name 'unsync', but found '{}'.", FieldVal.c_str()));
+		}
+	}
+
+	if (auto& Field = JsonObject["name"]; Field.is_string())
+	{
+		Result.Name = Field.string_value();
+	}
+
+	if (auto& Field = JsonObject["version"]; Field.is_string())
+	{
+		Result.VersionNumber = Field.string_value();
+	}
+
+	if (auto& Field = JsonObject["git"]; Field.is_string())
+	{
+		Result.VersionGit = Field.string_value();
+	}
+
+	if (auto& Field = JsonObject["session"]; Field.is_string())
+	{
+		Result.SessionId = Field.string_value();
+	}
+
+	if (auto& Field = JsonObject["features"]; Field.is_array())
+	{
+		Result.FeatureNames.reserve(Field.array_items().size());
+		for (auto& Elem : Field.array_items())
+		{
+			if (Elem.is_string())
+			{
+				Result.FeatureNames.push_back(Elem.string_value());
+
+				if (Elem.string_value() == "telemetry")
+				{
+					Result.Features.bTelemetry = true;
+				}
+
+				if (Elem.string_value() == "mirrors")
+				{
+					Result.Features.bMirrors = true;
+				}
+			}
+		}
+	}
+
+	return ResultOk(std::move(Result));
+}
+
+void
+FUnsyncProtocolImpl::SendTelemetryEvent(const FRemoteDesc& RemoteDesc, const FTelemetryEventSyncComplete& Event)
+{
+	using namespace json11;
+	json11::Json::object Obj;
+
+	if (!Event.Session.empty())
+	{
+		Obj["session"] = Event.Session;
+	}
+
+	Obj["type"]		 = "sync_complete";
+	Obj["source"]	 = Event.Source;
+	Obj["total_mb"]	 = SizeMb(Event.TotalBytes);  // use size in megabytes due lack of JSON 64 bit int support and better human readability
+	Obj["source_mb"] = SizeMb(Event.SourceBytes);
+	Obj["base_mb"]	 = SizeMb(Event.BaseBytes);
+	Obj["files_skipped"] = int(Event.SkippedFiles);
+	Obj["files_full"]	 = int(Event.FullCopyFiles);
+	Obj["files_partial"] = int(Event.PartialCopyFiles);
+	Obj["elapsed"]		 = Event.Elapsed;
+	Obj["success"]		 = Event.bSuccess;
+
+	std::string EventJson = Json(Obj).dump();
+
+	FBufferView EventJsonView;
+	EventJsonView.Data = (const uint8*)EventJson.c_str();
+	EventJsonView.Size = (uint64)EventJson.length();
+
+	const char* Url = "/api/v1/telemetry";
+	HttpRequest(RemoteDesc, EHttpMethod::POST, Url, EHttpContentType::Application_Json, EventJsonView);
+}
+
 FUnsyncProtocolImpl::~FUnsyncProtocolImpl()
 {
 	if (IsValid())
@@ -456,6 +584,31 @@ FProxyPool::FProxyPool(const FRemoteDesc& InRemoteDesc)
 , RemoteDesc(InRemoteDesc)
 , bValid(InRemoteDesc.IsValid())
 {
+	if (bValid && RemoteDesc.Protocol == EProtocolFlavor::Unsync)
+	{
+		UNSYNC_VERBOSE(L"Connecting to %hs server '%hs:%d' ...",
+			ToString(RemoteDesc.Protocol),
+			RemoteDesc.HostAddress.c_str(),
+			RemoteDesc.HostPort);
+
+		TResult<FUnsyncProtocolImpl::FHelloResponse> Response = FUnsyncProtocolImpl::QueryHello(RemoteDesc);
+
+		if (Response.IsError())
+		{
+			LogError(Response.GetError());
+		}
+
+		const FUnsyncProtocolImpl::FHelloResponse& Data = Response.GetData();
+		UNSYNC_VERBOSE(L"Connection established. Server name: %hs, version: %hs, git: %hs.", 
+			Data.Name.empty() ? "unknown" : Data.Name.c_str(),
+			Data.VersionNumber.empty() ? "unknown" : Data.VersionNumber.c_str(),
+			Data.VersionGit.empty() ? "unknown" : Data.VersionGit.c_str());
+
+		Features = Data.Features;
+		SessionId = Data.SessionId;
+
+		bValid = Response.IsOk();
+	}
 }
 
 std::unique_ptr<FProxy>
@@ -476,7 +629,7 @@ FProxyPool::Alloc()
 
 	if (!Result || !Result->IsValid())
 	{
-		Result = std::make_unique<FProxy>(RemoteDesc, &RequestMap);
+		Result = std::make_unique<FProxy>(RemoteDesc, Features, &RequestMap);
 	}
 
 	return Result;
@@ -511,6 +664,15 @@ FProxyPool::BuildFileBlockRequests(const FPath& OriginalFilePath, const FPath& R
 	RequestMap.AddFileBlocks(OriginalFilePath, ResolvedFilePath, FileManifest);
 }
 
+void
+FProxyPool::SendTelemetryEvent(const FTelemetryEventSyncComplete& Event)
+{
+	if (RemoteDesc.Protocol == EProtocolFlavor::Unsync && Features.bTelemetry)
+	{
+		FUnsyncProtocolImpl::SendTelemetryEvent(RemoteDesc, Event);
+	}
+}
+
 void
 FProxyPool::InitRequestMap(EStrongHashAlgorithmID InStrongHasher)
 {
diff --git a/Engine/Source/Programs/Unsync/Private/UnsyncProxy.h b/Engine/Source/Programs/Unsync/Private/UnsyncProxy.h
--- a/Engine/Source/Programs/Unsync/Private/UnsyncProxy.h
+++ b/Engine/Source/Programs/Unsync/Private/UnsyncProxy.h
@@ -73,6 +73,26 @@ private:
 	std::unordered_map<FHash128, FMacroBlockRequest> MacroBlockRequests;
 };
 
+struct FRemoteProtocolFeatures
+{
+	bool bTelemetry = false;
+	bool bMirrors = false;
+};
+
+struct FTelemetryEventSyncComplete
+{
+	std::string Session;
+	std::string Source;
+	uint64 TotalBytes = 0;
+	uint64 SourceBytes = 0;
+	uint64 BaseBytes = 0;
+	uint32 SkippedFiles = 0;
+	uint32 FullCopyFiles = 0;
+	uint32 PartialCopyFiles = 0;
+	double Elapsed = 0;
+	bool bSuccess = false;
+};
+
 struct FRemoteProtocolBase
 {
 	FRemoteProtocolBase(const FRemoteDesc& InRemoteDesc, const FBlockRequestMap* InRequestMap)
@@ -100,7 +120,7 @@ struct FRemoteProtocolBase
 class FProxy
 {
 public:
-	FProxy(const FRemoteDesc& InRemoteDesc, const FBlockRequestMap* InRequestMap);
+	FProxy(const FRemoteDesc& InRemoteDesc, const FRemoteProtocolFeatures& InFeatures, const FBlockRequestMap* InRequestMap);
 	~FProxy();
 
 	bool Contains(const FDirectoryManifest& Manifest);
@@ -131,10 +151,18 @@ public:
 	void InitRequestMap(EStrongHashAlgorithmID InStrongHasher);
 	void BuildFileBlockRequests(const FPath& OriginalFilePath, const FPath& ResolvedFilePath, const FFileManifest& FileManifest);
 
+	const FRemoteProtocolFeatures& GetFeatures() const { return Features; }
+	const std::string& GetSessionId() const { return SessionId; }
+
+	void SendTelemetryEvent(const FTelemetryEventSyncComplete& Event);
+
 private:
 	std::vector<std::unique_ptr<FProxy>> Pool;
 	bool								 bValid = true;
 
+	FRemoteProtocolFeatures Features;
+	std::string SessionId;
+
 	FBlockRequestMap RequestMap;
 
 	std::mutex Mutex;
