unsync - Add --overlay option to sync command to copy files from multiple sources into a single merged target, bump version to 1.0.43

diff --git a/Engine/Source/Programs/Unsync/Private/UnsyncCmdSync.cpp b/Engine/Source/Programs/Unsync/Private/UnsyncCmdSync.cpp
--- a/Engine/Source/Programs/Unsync/Private/UnsyncCmdSync.cpp
+++ b/Engine/Source/Programs/Unsync/Private/UnsyncCmdSync.cpp
@@ -9,15 +9,22 @@ namespace unsync {
 int32  // TODO: return a TResult
 CmdSync(const FCmdSyncOptions& Options)
 {
+	auto ResolvePath = [&Options](const FPath& Path) -> FPath
+	{
+		return Options.Filter ? Options.Filter->Resolve(Path) : Path;
+	};
+
 	std::error_code ErrorCode	   = {};
-	FPath			ResolvedSource = Options.Filter ? Options.Filter->Resolve(Options.Source) : Options.Source;
+	FPath			ResolvedSource = ResolvePath(Options.Source);
 
-	UNSYNC_VERBOSE(L"Sync source: '%ls'", Options.Source.wstring().c_str());
-	if (Options.Source != ResolvedSource)
+	if (Options.Source == ResolvedSource)
 	{
-		UNSYNC_VERBOSE(L"-- resolved: '%ls'", ResolvedSource.wstring().c_str());
+		UNSYNC_VERBOSE(L"Sync source: '%ls'", Options.Source.wstring().c_str());
+	}
+	else
+	{
+		UNSYNC_VERBOSE(L"Sync source: '%ls' ('%ls')", Options.Source.wstring().c_str(), ResolvedSource.wstring().c_str());
 	}
-	UNSYNC_VERBOSE(L"Sync target: '%ls'", Options.Target.wstring().c_str());
 
 	const bool bSourcePathExists	 = PathExists(ResolvedSource, ErrorCode);
 	const bool bSourceIsDirectory	 = bSourcePathExists && unsync::IsDirectory(ResolvedSource);
@@ -25,6 +32,38 @@ CmdSync(const FCmdSyncOptions& Options)
 
 	bool bSourceFileSystemRequired = !bSourceIsManifestHash;
 
+	if (!Options.Overlays.empty())
+	{
+		for (const FPath& Entry : Options.Overlays)
+		{
+			FPath ResolvedEntry = ResolvePath(Entry);
+			if (ResolvedEntry == Entry)
+			{
+				UNSYNC_VERBOSE(L"Sync overlay: '%ls'", ResolvedEntry.wstring().c_str());
+			}
+			else
+			{
+				UNSYNC_VERBOSE(L"Sync overlay: '%ls' ('%ls')", Entry.wstring().c_str(), ResolvedEntry.wstring().c_str());
+			}
+		}
+
+		bSourceFileSystemRequired = true;
+
+		if (!bSourcePathExists || !bSourceIsDirectory)
+		{
+			UNSYNC_ERROR(L"Sync overlay option requires sync source to be a directory that exists on disk.");
+			return 1;
+		}
+
+		if (bSourceIsManifestHash)
+		{
+			UNSYNC_ERROR(L"Sync overlay option is not compatible with manifest override.");
+			return 1;
+		}
+	}
+
+	UNSYNC_VERBOSE(L"Sync target: '%ls'", Options.Target.wstring().c_str());
+
 	if (!Options.SourceManifestOverride.empty())
 	{
 		UNSYNC_VERBOSE(L"Manifest override: %ls", Options.SourceManifestOverride.wstring().c_str());
@@ -64,6 +103,7 @@ CmdSync(const FCmdSyncOptions& Options)
 			SyncOptions.Source				   = Options.Source;
 			SyncOptions.Base				   = Options.Target;  // read base data from existing target
 			SyncOptions.Target				   = Options.Target;
+			SyncOptions.Overlays			   = Options.Overlays;
 			SyncOptions.SourceManifestOverride = Options.SourceManifestOverride;
 			SyncOptions.Remote				   = &Options.Remote;
 			SyncOptions.SyncFilter			   = Options.Filter;
diff --git a/Engine/Source/Programs/Unsync/Private/UnsyncCmdSync.h b/Engine/Source/Programs/Unsync/Private/UnsyncCmdSync.h
--- a/Engine/Source/Programs/Unsync/Private/UnsyncCmdSync.h
+++ b/Engine/Source/Programs/Unsync/Private/UnsyncCmdSync.h
@@ -16,6 +16,8 @@ struct FCmdSyncOptions
 	FPath Target;
 	FPath SourceManifestOverride;
 
+	std::vector<FPath> Overlays;
+
 	FRemoteDesc Remote;
 
 	bool bFullDifference = false;
diff --git a/Engine/Source/Programs/Unsync/Private/UnsyncCore.cpp b/Engine/Source/Programs/Unsync/Private/UnsyncCore.cpp
--- a/Engine/Source/Programs/Unsync/Private/UnsyncCore.cpp
+++ b/Engine/Source/Programs/Unsync/Private/UnsyncCore.cpp
@@ -25,7 +25,7 @@ UNSYNC_THIRD_PARTY_INCLUDES_START
 #include <flat_hash_map.hpp>
 UNSYNC_THIRD_PARTY_INCLUDES_END
 
-#define UNSYNC_VERSION_STR "1.0.42"
+#define UNSYNC_VERSION_STR "1.0.43"
 
 namespace unsync {
 
@@ -2040,6 +2040,13 @@ UpdateDirectoryManifestBlocks(FDirectoryManifest& Result, const FPath& Root, uin
 	}
 }
 
+static bool AlgorithmOptionsCompatible(const FAlgorithmOptions& A, const FAlgorithmOptions& B)
+{
+	return A.StrongHashAlgorithmId == B.StrongHashAlgorithmId
+		&& B.WeakHashAlgorithmId == B.WeakHashAlgorithmId
+		&& B.ChunkingAlgorithmId == B.ChunkingAlgorithmId;
+}
+
 bool
 LoadOrCreateDirectoryManifest(FDirectoryManifest& Result, const FPath& Root, uint32 BlockSize, FAlgorithmOptions Algorithm)
 {
@@ -2058,8 +2065,7 @@ LoadOrCreateDirectoryManifest(FDirectoryManifest& Result, const FPath& Root, uin
 	}
 
 	const bool bExistingManifestLoaded = bManifestFileExists && LoadDirectoryManifest(OldDirectoryManifest, Root, DirectoryManifestPath);
-	const bool bExistingManifestCompatible = OldDirectoryManifest.Options.StrongHashAlgorithmId == Algorithm.StrongHashAlgorithmId &&
-											 OldDirectoryManifest.Options.WeakHashAlgorithmId == Algorithm.WeakHashAlgorithmId;
+	const bool bExistingManifestCompatible = AlgorithmOptionsCompatible(OldDirectoryManifest.Options, Algorithm);
 
 	if (bExistingManifestLoaded && bExistingManifestCompatible)
 	{
@@ -2621,6 +2627,12 @@ IsNonCaseSensitiveFileSystem(const FPath& ExistingPath)
 	}
 }
 
+static bool
+IsCaseSensitiveFileSystem(const FPath& ExistingPath)
+{
+	return !IsNonCaseSensitiveFileSystem(ExistingPath);
+}
+
 struct FPendingFileRename
 {
 	std::wstring Old;
@@ -2770,6 +2782,62 @@ FixFileNameCases(const FPath& RootPath, const std::vector<FPendingFileRename>& P
 	return true;
 }
 
+static bool
+MergeManifests(FDirectoryManifest& Existing, const FDirectoryManifest& Other, bool bCaseSensitive)
+{
+	if (!Existing.IsValid())
+	{
+		Existing = Other;
+		return true;
+	}
+
+	if (!AlgorithmOptionsCompatible(Existing.Options, Other.Options))
+	{
+		UNSYNC_ERROR("Trying to merge incompatible manifests (diff algorithm options do not match)");
+		return false;
+	}
+
+	if (bCaseSensitive)
+	{
+		// Trivial case: just replace existing entries
+		for (const auto& OtherFile : Other.Files)
+		{
+			Existing.Files[OtherFile.first] = OtherFile.second;
+		}
+	}
+	else
+	{
+		// Lookup table of lowercase -> original file name used to replace conflicting entries on non-case-sensitive filesystems
+		// TODO: Could potentially add case-sensitive/insensitive entry lookup helper functions to FDirectoryManifest itself in the future
+		std::unordered_map<std::wstring, std::wstring> ExistingFileNamesLowerCase;
+
+		for (auto& ExistingEntry : Existing.Files)
+		{
+			std::wstring FileNameLowerCase = StringToLower(ExistingEntry.first);
+			ExistingFileNamesLowerCase.insert(std::pair<std::wstring, std::wstring>(FileNameLowerCase, ExistingEntry.first));
+		}
+
+		for (const auto& OtherFile : Other.Files)
+		{
+			std::wstring OtherNameLowerCase = StringToLower(OtherFile.first);
+			auto LowerCaseEntry = ExistingFileNamesLowerCase.find(OtherNameLowerCase);
+			if (LowerCaseEntry != ExistingFileNamesLowerCase.end())
+			{
+				// Remove file with conflicting case and add entry from the other manifest instead
+				const std::wstring& ExistingNameOriginalCase = LowerCaseEntry->second;
+				Existing.Files.erase(ExistingNameOriginalCase);
+
+				// Update the lookup table entry to refer to the name we're about to insert
+				ExistingFileNamesLowerCase[LowerCaseEntry->first] = OtherFile.first;
+			}
+
+			Existing.Files[OtherFile.first] = OtherFile.second;
+		}
+	}
+
+	return true;
+}
+
 // Delete files from target directory that are not in the source directory manifest
 static void
 DeleteUnnecessaryFiles(const FPath&				 TargetDirectory,
@@ -2813,6 +2881,71 @@ ToPath(const std::wstring_view& Str)
 #endif	// UNSYNC_PLATFORM_UNIX
 }
 
+static bool
+LoadAndMergeSourceManifest(FDirectoryManifest& Output,
+						   const FPath&		   SourcePath,
+						   const FPath&		   TempPath,
+						   FSyncFilter*		   SyncFilter,
+						   const FPath&		   SourceManifestOverride,
+						   bool				   bCaseSensitiveTargetFileSystem)
+{
+	auto ResolvePath = [SyncFilter](const FPath& Filename) -> FPath { return SyncFilter ? SyncFilter->Resolve(Filename) : Filename; };
+
+	FDirectoryManifest LoadedManifest;
+
+	FPath SourceManifestRoot = SourcePath / ".unsync";
+	FPath SourceManifestPath = SourceManifestRoot / "manifest.bin";
+
+	SourceManifestPath = ResolvePath(SourceManifestPath);
+
+	if (!SourceManifestOverride.empty())
+	{
+		SourceManifestPath = SourceManifestOverride;
+	}
+
+	FHash128 SourcePathHash =
+		HashBlake3Bytes<FHash128>((const uint8*)SourcePath.native().c_str(), SourcePath.native().length() * sizeof(SourcePath.native()[0]));
+
+	std::string SourcePathHashStr	   = BytesToHexString(SourcePathHash.Data, sizeof(SourcePathHash.Data));
+	FPath		SourceManifestTempPath = TempPath / SourcePathHashStr;
+
+	if (!PathExists(SourceManifestPath))
+	{
+		UNSYNC_ERROR(L"Source manifest '%ls' does not exist", SourceManifestPath.wstring().c_str());
+		return false;
+	}
+
+	LogGlobalStatus(L"Caching source manifest");
+
+	UNSYNC_VERBOSE(L"Caching source manifest");
+	UNSYNC_VERBOSE(L" Source '%ls'", SourceManifestPath.wstring().c_str());
+	UNSYNC_VERBOSE(L" Target '%ls'", SourceManifestTempPath.wstring().c_str());
+	std::error_code CopyErrorCode = CopyFileIfNewer(SourceManifestPath, SourceManifestTempPath);
+	if (CopyErrorCode)
+	{
+		UNSYNC_LOG(L"Failed to copy manifest '%ls' to '%ls'",
+				   SourceManifestPath.wstring().c_str(),
+				   SourceManifestTempPath.wstring().c_str());
+		UNSYNC_ERROR(L"%hs (%d)", CopyErrorCode.message().c_str(), CopyErrorCode.value());
+		return false;
+	}
+
+	if (!LoadDirectoryManifest(LoadedManifest, SourcePath, SourceManifestTempPath))
+	{
+		UNSYNC_ERROR(L"Failed to load source directory manifest '%ls'", SourceManifestPath.wstring().c_str());
+
+		return false;
+	}
+
+	if (Output.IsValid() && !AlgorithmOptionsCompatible(Output.Options, LoadedManifest.Options))
+	{
+		UNSYNC_ERROR(L"Can't merge manifest '%ls' as it uses different algorithm options", SourcePath.wstring().c_str());
+		return false;
+	}
+
+	return MergeManifests(Output, LoadedManifest, bCaseSensitiveTargetFileSystem);
+}
+
 bool  // TODO: return a TResult
 SyncDirectory(const FSyncDirectoryOptions& SyncOptions)
 {
@@ -2877,56 +3010,32 @@ SyncDirectory(const FSyncDirectoryOptions& SyncOptions)
 	FDirectoryManifest SourceDirectoryManifest;
 	FPath			   SourceManifestTempPath;
 
+	const bool bCaseSensitiveTargetFileSystem = IsCaseSensitiveFileSystem(TargetTempPath);
+
 	if (bFileSystemSource || !SourceManifestOverride.empty())
 	{
-		FPath SourceManifestRoot = SourcePath / ".unsync";
-		FPath SourceManifestPath = SourceManifestRoot / "manifest.bin";
-
-		SourceManifestPath = ResolvePath(SourceManifestPath);
-
-		if (!SourceManifestOverride.empty())
-		{
-			SourceManifestPath = SourceManifestOverride;
-		}
-
-		FHash128 SourcePathHash = HashBlake3Bytes<FHash128>((const uint8*)SourcePath.native().c_str(),
-															SourcePath.native().length() * sizeof(SourcePath.native()[0]));
-
-		std::string SourcePathHashStr = BytesToHexString(SourcePathHash.Data, sizeof(SourcePathHash.Data));
-		SourceManifestTempPath		  = TargetTempPath / SourcePathHashStr;
-
-		if (!PathExists(SourceManifestPath))
-		{
-			UNSYNC_ERROR(L"Source manifest '%ls' does not exist", SourceManifestPath.wstring().c_str());
-			return false;
-		}
-
-		LogGlobalStatus(L"Caching source manifest");
-
-		UNSYNC_VERBOSE(L"Caching source manifest");
-		UNSYNC_VERBOSE(L" Source '%ls'", SourceManifestPath.wstring().c_str());
-		UNSYNC_VERBOSE(L" Target '%ls'", SourceManifestTempPath.wstring().c_str());
-		std::error_code CopyErrorCode = CopyFileIfNewer(SourceManifestPath, SourceManifestTempPath);
-		if (CopyErrorCode)
+		std::vector<FPath> AllSources;
+		AllSources.push_back(SourcePath);
+		for (const FPath& OverlayPath : SyncOptions.Overlays)
 		{
-			UNSYNC_LOG(L"Failed to copy manifest '%ls' to '%ls'",
-					   SourceManifestPath.wstring().c_str(),
-					   SourceManifestTempPath.wstring().c_str());
-			UNSYNC_ERROR(L"%hs (%d)", CopyErrorCode.message().c_str(), CopyErrorCode.value());
-			return false;
+			AllSources.push_back(OverlayPath);
 		}
 
-		if (!LoadDirectoryManifest(SourceDirectoryManifest, SourcePath, SourceManifestTempPath))
+		for (const FPath& ThisSourcePath : AllSources)
 		{
-			bSourceManifestOk = false;
-			UNSYNC_ERROR(L"Failed to load source directory manifest '%ls'", SourceManifestPath.wstring().c_str());
-
-			return false;
+			if (!LoadAndMergeSourceManifest(SourceDirectoryManifest,
+				ThisSourcePath,
+				TargetTempPath,
+				SyncFilter,
+				SourceManifestOverride,
+				bCaseSensitiveTargetFileSystem))
+			{
+				return false;
+			}
 		}
 	}
 	else
 	{
-		// TODO: download manifest from Jupiter
 		if (!ProxyPool.IsValid())
 		{
 			UNSYNC_ERROR(L"Remote server connection is required when syncing by manifest hash");
@@ -2970,7 +3079,7 @@ SyncDirectory(const FSyncDirectoryOptions& SyncOptions)
 	FAlgorithmOptions  Algorithm			   = SourceDirectoryManifest.Options;
 	FDirectoryManifest TargetDirectoryManifest = CreateDirectoryManifest(TargetPath, 0, Algorithm);
 
-	if (IsNonCaseSensitiveFileSystem(TargetTempPath))
+	if (!bCaseSensitiveTargetFileSystem)
 	{
 		std::vector<FPendingFileRename> PendingRenames;
 		PendingRenames = FixManifestFileNameCases(TargetDirectoryManifest, SourceDirectoryManifest);
@@ -3037,10 +3146,7 @@ SyncDirectory(const FSyncDirectoryOptions& SyncOptions)
 		&& PathExists(BaseManifestPath))
 	{
 		bBaseDirectoryManifestValid = LoadDirectoryManifest(BaseDirectoryManifest, BasePath, BaseManifestPath);
-		if (bBaseDirectoryManifestValid &&
-			SourceDirectoryManifest.Options.ChunkingAlgorithmId == TargetDirectoryManifest.Options.ChunkingAlgorithmId &&
-			SourceDirectoryManifest.Options.StrongHashAlgorithmId == TargetDirectoryManifest.Options.StrongHashAlgorithmId &&
-			SourceDirectoryManifest.Options.WeakHashAlgorithmId == TargetDirectoryManifest.Options.WeakHashAlgorithmId)
+		if (bBaseDirectoryManifestValid && AlgorithmOptionsCompatible(SourceDirectoryManifest.Options, TargetDirectoryManifest.Options))
 		{
 			bQuickDifferencePossible = true;
 		}
@@ -3081,7 +3187,7 @@ SyncDirectory(const FSyncDirectoryOptions& SyncOptions)
 			}
 		}
 
-		FPath SourceFilePath = SourcePath / ToPath(SourceManifestIt.first);
+		FPath SourceFilePath = SourceManifestIt.second.CurrentPath;
 		FPath BaseFilePath	 = BasePath / ToPath(SourceManifestIt.first);
 		FPath TargetFilePath = TargetPath / ToPath(SourceManifestIt.first);
 
@@ -3543,16 +3649,11 @@ SyncDirectory(const FSyncDirectoryOptions& SyncOptions)
 	// It can be used to speed up the diffing process during next sync.
 	if (bFileSystemSource && bSyncSucceeded && !GDryRun)
 	{
-		// TODO: cache the manfiest if it's downloaded from Jupiter also
-
-		std::error_code ErrorCode;
-		bool			bCopyOk = FileCopyOverwrite(SourceManifestTempPath, TargetManifestPath, ErrorCode);
+		bool bSaveOk = SaveDirectoryManifest(SourceDirectoryManifest, TargetManifestPath);
 
-		if (!bCopyOk)
+		if (!bSaveOk)
 		{
-			UNSYNC_WARNING(L"Failed to save manifest after sync. System error code: %hs (%d).",
-						   ErrorCode.message().c_str(),
-						   ErrorCode.value());
+			UNSYNC_ERROR(L"Failed to save manifest after sync");
 		}
 	}
 
diff --git a/Engine/Source/Programs/Unsync/Private/UnsyncCore.h b/Engine/Source/Programs/Unsync/Private/UnsyncCore.h
--- a/Engine/Source/Programs/Unsync/Private/UnsyncCore.h
+++ b/Engine/Source/Programs/Unsync/Private/UnsyncCore.h
@@ -163,6 +163,11 @@ struct FDirectoryManifest
 	// runtime data
 	FAlgorithmOptions Options = {};
 	uint64			  Version = 0;
+
+	bool IsValid() const
+	{
+		return Version != EVersions::Invalid;
+	}
 };
 
 struct FDirectoryManifestInfo
@@ -340,6 +345,7 @@ struct FSyncDirectoryOptions
 	FPath			   Source;					// remote data location
 	FPath			   Target;					// output target location
 	FPath			   Base;					// base data location, which typically is the same as sync target
+	std::vector<FPath> Overlays;				// extra source directories to overlay over primary (add extra files, replace existing files)
 	FPath			   SourceManifestOverride;	// force the manifest to be read from a specified file instead of source directory
 	FSyncFilter*	   SyncFilter = nullptr;	// filter callback for partial sync support
 	const FRemoteDesc* Remote	  = nullptr;	// unsync proxy server connection settings
diff --git a/Engine/Source/Programs/Unsync/Private/UnsyncMain.cpp b/Engine/Source/Programs/Unsync/Private/UnsyncMain.cpp
--- a/Engine/Source/Programs/Unsync/Private/UnsyncMain.cpp
+++ b/Engine/Source/Programs/Unsync/Private/UnsyncMain.cpp
@@ -54,6 +54,7 @@ InnerMain(int Argc, char** Argv)
 	std::string				 InputFilename2Utf8;
 	std::string				 SourceManifestFilenameUtf8;
 	std::vector<std::string> ExcludeFilterArrayUtf8;
+	std::vector<std::string> OverlayArrayUtf8;
 	std::string				 RemoteAddressUtf8;
 	std::string				 PreferredDfsUtf8;
 	std::string				 WeakHashUtf8	= "buzhash";
@@ -64,18 +65,18 @@ InnerMain(int Argc, char** Argv)
 	std::string				 ProtocolName = "jupiter";
 	std::string				 HttpHeaderFilenameUtf8;
 	std::string				 QueryStringUtf8;
-	bool					 bForceOperation		= false;
-	bool					 bAllowInsecureTls		= false;
-	bool					 bUseTls				= false;
-	bool					 bUseDebugMode			= false;
-	bool					 bIncrementalMode		= false;
-	bool					 bNoOutputValidation	= false;
-	bool					 bNoCleanupAfterSync	= false;
-	bool					 bFullSourceScan		= false;
-	bool					 bFullDifference		= false;
-	int32					 CompressionLevel		= 3;
-	uint32					 DiffBlockSize			= uint32(4_KB);
-	uint32					 HashOrSyncBlockSize	= uint32(64_KB);
+	bool					 bForceOperation	 = false;
+	bool					 bAllowInsecureTls	 = false;
+	bool					 bUseTls			 = false;
+	bool					 bUseDebugMode		 = false;
+	bool					 bIncrementalMode	 = false;
+	bool					 bNoOutputValidation = false;
+	bool					 bNoCleanupAfterSync = false;
+	bool					 bFullSourceScan	 = false;
+	bool					 bFullDifference	 = false;
+	int32					 CompressionLevel	 = 3;
+	uint32					 DiffBlockSize		 = uint32(4_KB);
+	uint32					 HashOrSyncBlockSize = uint32(64_KB);
 
 	struct FDeprecatedOptions
 	{
@@ -159,6 +160,7 @@ InnerMain(int Argc, char** Argv)
 						RemoteAddressUtf8,
 						"FProxy server address ([transport://]address[:port][/request][#namespace])");
 	SubSync->add_option("--dfs", PreferredDfsUtf8, "Preferred DFS mirror (matched by sub-string)");
+	SubSync->add_option("--overlay", OverlayArrayUtf8, "Additional source directory to sync (keep unique files from all sources, overwrite conflicting files with overlay source)");
 	SubSync->add_option("--exclude", ExcludeFilterArrayUtf8, "Exclude filenames that contain specified words (comma separated)");
 #if UNSYNC_USE_TLS
 	AddTlsOptions(SubSync);
@@ -617,6 +619,11 @@ InnerMain(int Argc, char** Argv)
 		SyncOptions.Filter				   = &SyncFilter;
 		SyncOptions.bValidateTargetFiles   = !bNoOutputValidation;
 
+		for (const std::string& Entry : OverlayArrayUtf8)
+		{
+			SyncOptions.Overlays.push_back(NormalizeFilenameUtf8(Entry));
+		}
+
 		return CmdSync(SyncOptions);
 	}
 	else if (Cli.got_subcommand(SubPatch))
