unsync - Add client version string to telemetry data

diff --git a/Engine/Source/Programs/Unsync/Private/UnsyncCore.cpp b/Engine/Source/Programs/Unsync/Private/UnsyncCore.cpp
--- a/Engine/Source/Programs/Unsync/Private/UnsyncCore.cpp
+++ b/Engine/Source/Programs/Unsync/Private/UnsyncCore.cpp
@@ -3668,6 +3668,7 @@ SyncDirectory(const FSyncDirectoryOptions& SyncOptions)
 	{
 		FTelemetryEventSyncComplete Event;
 
+		Event.ClientVersion	   = GetVersionString();
 		Event.Session		   = ProxyPool.GetSessionId();
 		Event.Source		   = ConvertWideToUtf8(SourcePath.wstring());
 		Event.TotalBytes	   = TotalSourceSize;
diff --git a/Engine/Source/Programs/Unsync/Private/UnsyncProxy.cpp b/Engine/Source/Programs/Unsync/Private/UnsyncProxy.cpp
--- a/Engine/Source/Programs/Unsync/Private/UnsyncProxy.cpp
+++ b/Engine/Source/Programs/Unsync/Private/UnsyncProxy.cpp
@@ -438,8 +438,9 @@ FUnsyncProtocolImpl::SendTelemetryEvent(const FRemoteDesc& RemoteDesc, const FTe
 		Obj["session"] = Event.Session;
 	}
 
-	Obj["type"]		 = "sync_complete";
-	Obj["source"]	 = Event.Source;
+	Obj["type"]			  = "sync_complete";
+	Obj["client_version"] = Event.ClientVersion;
+	Obj["source"]		  = Event.Source;
 	Obj["total_mb"]	 = SizeMb(Event.TotalBytes);  // use size in megabytes due lack of JSON 64 bit int support and better human readability
 	Obj["source_mb"] = SizeMb(Event.SourceBytes);
 	Obj["base_mb"]	 = SizeMb(Event.BaseBytes);
diff --git a/Engine/Source/Programs/Unsync/Private/UnsyncProxy.h b/Engine/Source/Programs/Unsync/Private/UnsyncProxy.h
--- a/Engine/Source/Programs/Unsync/Private/UnsyncProxy.h
+++ b/Engine/Source/Programs/Unsync/Private/UnsyncProxy.h
@@ -81,6 +81,7 @@ struct FRemoteProtocolFeatures
 
 struct FTelemetryEventSyncComplete
 {
+	std::string ClientVersion;
 	std::string Session;
 	std::string Source;
 	uint64 TotalBytes = 0;
