unsync - Fix crash when initial proxy handshake fails

Also use plain (non-TLS) connection for handshake by default

diff --git a/Engine/Source/Programs/Unsync/Private/UnsyncCore.cpp b/Engine/Source/Programs/Unsync/Private/UnsyncCore.cpp
--- a/Engine/Source/Programs/Unsync/Private/UnsyncCore.cpp
+++ b/Engine/Source/Programs/Unsync/Private/UnsyncCore.cpp
@@ -25,7 +25,7 @@ UNSYNC_THIRD_PARTY_INCLUDES_START
 #include <flat_hash_map.hpp>
 UNSYNC_THIRD_PARTY_INCLUDES_END
 
-#define UNSYNC_VERSION_STR "1.0.43"
+#define UNSYNC_VERSION_STR "1.0.44"
 
 namespace unsync {
 
@@ -3399,7 +3399,7 @@ SyncDirectory(const FSyncDirectoryOptions& SyncOptions)
 
 		std::unique_ptr<FProxy> Proxy = ProxyPool.Alloc();
 
-		if (Proxy->IsValid())
+		if (Proxy.get() && Proxy->IsValid())
 		{
 			// TODO: report TLS status
 			// ESocketSecurity security = proxy->get_socket_security();
diff --git a/Engine/Source/Programs/Unsync/Private/UnsyncHttp.cpp b/Engine/Source/Programs/Unsync/Private/UnsyncHttp.cpp
--- a/Engine/Source/Programs/Unsync/Private/UnsyncHttp.cpp
+++ b/Engine/Source/Programs/Unsync/Private/UnsyncHttp.cpp
@@ -237,8 +237,8 @@ HttpRequest(const FRemoteDesc& RemoteDesc,
 			EHttpContentType   PayloadContentType,
 			FBufferView		   Payload)
 {
-	FTlsClientSettings TlsSettings = RemoteDesc.GetTlsClientSettings();
-	FHttpConnection	   Connection(RemoteDesc.HostAddress, RemoteDesc.HostPort, &TlsSettings);
+	const FTlsClientSettings* TlsSettings = nullptr;
+	FHttpConnection	   Connection(RemoteDesc.HostAddress, RemoteDesc.HostPort, TlsSettings);
 
 	FHttpRequest Request;
 
diff --git a/Engine/Source/Programs/Unsync/Private/UnsyncProxy.cpp b/Engine/Source/Programs/Unsync/Private/UnsyncProxy.cpp
--- a/Engine/Source/Programs/Unsync/Private/UnsyncProxy.cpp
+++ b/Engine/Source/Programs/Unsync/Private/UnsyncProxy.cpp
@@ -598,15 +598,17 @@ FProxyPool::FProxyPool(const FRemoteDesc& InRemoteDesc)
 		{
 			LogError(Response.GetError());
 		}
-
-		const FUnsyncProtocolImpl::FHelloResponse& Data = Response.GetData();
-		UNSYNC_VERBOSE(L"Connection established. Server name: %hs, version: %hs, git: %hs.", 
-			Data.Name.empty() ? "unknown" : Data.Name.c_str(),
-			Data.VersionNumber.empty() ? "unknown" : Data.VersionNumber.c_str(),
-			Data.VersionGit.empty() ? "unknown" : Data.VersionGit.c_str());
-
-		Features = Data.Features;
-		SessionId = Data.SessionId;
+		else
+		{
+			const FUnsyncProtocolImpl::FHelloResponse& Data = Response.GetData();
+			UNSYNC_VERBOSE(L"Connection established. Server name: %hs, version: %hs, git: %hs.", 
+				Data.Name.empty() ? "unknown" : Data.Name.c_str(),
+				Data.VersionNumber.empty() ? "unknown" : Data.VersionNumber.c_str(),
+				Data.VersionGit.empty() ? "unknown" : Data.VersionGit.c_str());
+
+			Features = Data.Features;
+			SessionId = Data.SessionId;
+		}
 
 		bValid = Response.IsOk();
 	}
