unsync - Use TLS during handshake if running with --tls option, skip all certificate validation with --insecure option

diff --git a/Engine/Source/Programs/Unsync/Private/UnsyncCore.cpp b/Engine/Source/Programs/Unsync/Private/UnsyncCore.cpp
--- a/Engine/Source/Programs/Unsync/Private/UnsyncCore.cpp
+++ b/Engine/Source/Programs/Unsync/Private/UnsyncCore.cpp
@@ -3388,7 +3388,7 @@ SyncDirectory(const FSyncDirectoryOptions& SyncOptions)
 
 	LogGlobalProgress();
 
-	if (ProxySettings.IsValid())
+	if (ProxySettings.IsValid() && ProxyPool.IsValid())
 	{
 		LogGlobalStatus(L"Connecting to server");
 		UNSYNC_VERBOSE(L"Connecting to %hs server '%hs:%d' ...",
@@ -3432,14 +3432,13 @@ SyncDirectory(const FSyncDirectoryOptions& SyncOptions)
 			Proxy = nullptr;
 			ProxyPool.Invalidate();
 		}
+	}
 
-		if (!ProxyPool.IsValid())
-		{
-			// TODO: bail out if remote connection is required for the download,
-			// such as when downloading data purely from Jupiter.
-
-			UNSYNC_WARNING(L"Attempting to sync without remote server connection");
-		}
+	if (ProxySettings.IsValid() && !ProxyPool.IsValid())
+	{
+		// TODO: bail out if remote connection is required for the download,
+		// such as when downloading data purely from Jupiter.
+		UNSYNC_WARNING(L"Attempting to sync without remote server connection");
 	}
 
 	LogGlobalStatus(L"Copying files");
diff --git a/Engine/Source/Programs/Unsync/Private/UnsyncHttp.cpp b/Engine/Source/Programs/Unsync/Private/UnsyncHttp.cpp
--- a/Engine/Source/Programs/Unsync/Private/UnsyncHttp.cpp
+++ b/Engine/Source/Programs/Unsync/Private/UnsyncHttp.cpp
@@ -237,8 +237,8 @@ HttpRequest(const FRemoteDesc& RemoteDesc,
 			EHttpContentType   PayloadContentType,
 			FBufferView		   Payload)
 {
-	const FTlsClientSettings* TlsSettings = nullptr;
-	FHttpConnection	   Connection(RemoteDesc.HostAddress, RemoteDesc.HostPort, TlsSettings);
+	FTlsClientSettings TlsSettings = RemoteDesc.GetTlsClientSettings();
+	FHttpConnection	   Connection(RemoteDesc.HostAddress, RemoteDesc.HostPort, RemoteDesc.bTlsEnable ? &TlsSettings : nullptr);
 
 	FHttpRequest Request;
 
@@ -495,7 +495,8 @@ FHttpConnection::Open()
 #if UNSYNC_USE_TLS
 		FTlsClientSettings ClientSettings;
 		ClientSettings.bVerifyCertificate = bTlsVerifyCertificate;
-		if (!TlsSubject.empty())
+
+		if (bTlsVerifyCertificate && !TlsSubject.empty())
 		{
 			ClientSettings.Subject = TlsSubject.c_str();
 		}
diff --git a/Engine/Source/Programs/Unsync/Private/UnsyncLog.cpp b/Engine/Source/Programs/Unsync/Private/UnsyncLog.cpp
--- a/Engine/Source/Programs/Unsync/Private/UnsyncLog.cpp
+++ b/Engine/Source/Programs/Unsync/Private/UnsyncLog.cpp
@@ -301,22 +301,26 @@ LogError(const FError& E)
 			break;
 	}
 
+	const wchar_t* ContextStr = E.Context.empty() ? nullptr : E.Context.c_str();
+
 	if (ErrorDescStr)
 	{
 		LogPrintf(ELogLevel::Error,
-				  L"%hs code: %d (%hs). Context: %ls\n",
+				  L"%hs code: %d (%hs).%ls%ls\n",
 				  ErrorKindStr,
 				  E.Code,
 				  ErrorDescStr,
-				  E.Context.empty() ? L"<none>" : E.Context.c_str());
+				  ContextStr ? L" Context: " : L"",
+				  E.Context.empty() ? L"" : E.Context.c_str());
 	}
 	else
 	{
 		LogPrintf(ELogLevel::Error,
-				  L"%hs code: %d. Context: %ls\n",
+				  L"%hs code: %d.%ls%ls\n",
 				  ErrorKindStr,
 				  E.Code,
-				  E.Context.empty() ? L"<none>" : E.Context.c_str());
+				  ContextStr ? L" Context: " : L"",
+				  E.Context.empty() ? L"" : E.Context.c_str());
 	}
 }
 
diff --git a/Engine/Source/Programs/Unsync/Private/UnsyncMain.cpp b/Engine/Source/Programs/Unsync/Private/UnsyncMain.cpp
--- a/Engine/Source/Programs/Unsync/Private/UnsyncMain.cpp
+++ b/Engine/Source/Programs/Unsync/Private/UnsyncMain.cpp
@@ -554,12 +554,14 @@ InnerMain(int Argc, char** Argv)
 
 	if (bUseTls)
 	{
-		RemoteDesc.bTlsEnable = bUseTls;
+		RemoteDesc.bTlsEnable			 = true;
+		RemoteDesc.bTlsVerifyCertificate = true;
 	}
 
 	if (bAllowInsecureTls)
 	{
 		RemoteDesc.bTlsVerifyCertificate = false;
+		RemoteDesc.TlsSubject			 = {};
 		UNSYNC_WARNING(L"Remote server certificate verification is disabled.");
 	}
 
diff --git a/Engine/Source/Programs/Unsync/Private/UnsyncProxy.cpp b/Engine/Source/Programs/Unsync/Private/UnsyncProxy.cpp
--- a/Engine/Source/Programs/Unsync/Private/UnsyncProxy.cpp
+++ b/Engine/Source/Programs/Unsync/Private/UnsyncProxy.cpp
@@ -357,7 +357,7 @@ TResult<FUnsyncProtocolImpl::FHelloResponse> FUnsyncProtocolImpl::QueryHello(con
 	if (!Response.Success())
 	{
 		UNSYNC_ERROR(L"Failed to establish connection to UNSYNC server. Error code: %d.", Response.Code);
-		return HttpError(Response.Code);
+		return HttpError(fmt::format("{}:{}{}", RemoteDesc.HostAddress.c_str(), RemoteDesc.HostPort, Url), Response.Code);
 	}
 
 	FHelloResponse Result;
