unsync - Update CMake to use blake3 library via vcpkg

diff --git a/Engine/Source/Programs/Unsync/CMakeLists.txt b/Engine/Source/Programs/Unsync/CMakeLists.txt
--- a/Engine/Source/Programs/Unsync/CMakeLists.txt
+++ b/Engine/Source/Programs/Unsync/CMakeLists.txt
@@ -89,16 +89,7 @@ set(src
 
 # Add third party code
 
-set(blake3dir ${CMAKE_CURRENT_SOURCE_DIR}/../../ThirdParty/BLAKE3/1.3.1/c)
-
 set(thirdparty
-	${blake3dir}/blake3_dispatch.c
-	${blake3dir}/blake3_impl.h
-	${blake3dir}/blake3_portable.c
-	${blake3dir}/blake3_sse41.c
-	${blake3dir}/blake3_avx2.c
-	${blake3dir}/blake3.c
-	${blake3dir}/blake3.h
 	ThirdParty/md5-sse2.cpp
 	ThirdParty/md5-sse2.h
 	ThirdParty/json11.cpp
@@ -120,18 +111,22 @@ find_package(fmt CONFIG REQUIRED)
 find_package(unofficial-http-parser CONFIG REQUIRED)
 find_package(zstd CONFIG REQUIRED)
 
-target_include_directories(unsync PRIVATE ${CMAKE_CURRENT_SOURCE_DIR}/thirdparty ${blake3dir})
-target_compile_definitions(unsync PRIVATE BLAKE3_NO_SSE2=1 BLAKE3_NO_AVX512=1)
+
+target_include_directories(unsync PRIVATE ${CMAKE_CURRENT_SOURCE_DIR}/thirdparty)
 
 target_link_directories(unsync PRIVATE ${VCPKG_INSTALLED_DIR}/${VCPKG_TARGET_TRIPLET}/lib)
 
-target_link_libraries(unsync PRIVATE 
+target_link_libraries(unsync PRIVATE
 	CLI11::CLI11
 	fmt::fmt-header-only # Header-only version as workaround for linker errors on Mac
 	unofficial::http_parser::http_parser
 	zstd::libzstd_static
 )
 
+find_path(BLAKE3_INCLUDE_DIR blake3.h REQUIRED)
+target_include_directories(unsync PRIVATE ${BLAKE3_INCLUDE_DIR})
+target_link_libraries(unsync PRIVATE blake3)
+
 if (UNSYNC_USE_TLS)
 	find_path(LIBRESSL_INCLUDE_DIR tls.h)
 	target_include_directories(unsync PRIVATE ${LIBRESSL_INCLUDE_DIR})
@@ -179,11 +174,6 @@ if(MSVC)
 
 	set_property(DIRECTORY ${CMAKE_CURRENT_SOURCE_DIR} PROPERTY VS_STARTUP_PROJECT unsync)
 
-else()
-
-	set_source_files_properties(${blake3dir}/blake3_sse41.c COMPILE_FLAGS -msse4.1)
-	set_source_files_properties(${blake3dir}/blake3_avx2.c  COMPILE_FLAGS -mavx2)
-
 endif()
 
 source_group("Private" FILES ${src})
diff --git a/Engine/Source/Programs/Unsync/vcpkg.json b/Engine/Source/Programs/Unsync/vcpkg.json
--- a/Engine/Source/Programs/Unsync/vcpkg.json
+++ b/Engine/Source/Programs/Unsync/vcpkg.json
@@ -2,6 +2,7 @@
   "name": "unsync",
   "version-string": "1",
   "dependencies": [
+    "blake3",
     "cli11",
     "fmt",
     "http-parser",
