unsync - Add a way to exclude certain files/directories from cleanup, bump version to 1.0.45

* Use --cleanup-exclude [comma,separated,words] option to skip deleting specific files/directories after sync (such as saved logs, configs, etc.)
* Use UNSYNC_CLEANUP_EXCLUDE environment variable to provide extra exclusion rules

diff --git a/Engine/Source/Programs/Unsync/Private/UnsyncCore.cpp b/Engine/Source/Programs/Unsync/Private/UnsyncCore.cpp
--- a/Engine/Source/Programs/Unsync/Private/UnsyncCore.cpp
+++ b/Engine/Source/Programs/Unsync/Private/UnsyncCore.cpp
@@ -25,7 +25,7 @@ UNSYNC_THIRD_PARTY_INCLUDES_START
 #include <flat_hash_map.hpp>
 UNSYNC_THIRD_PARTY_INCLUDES_END
 
-#define UNSYNC_VERSION_STR "1.0.44"
+#define UNSYNC_VERSION_STR "1.0.45"
 
 namespace unsync {
 
@@ -2842,14 +2842,33 @@ MergeManifests(FDirectoryManifest& Existing, const FDirectoryManifest& Other, bo
 static void
 DeleteUnnecessaryFiles(const FPath&				 TargetDirectory,
 					   const FDirectoryManifest& TargetDirectoryManifest,
-					   const FDirectoryManifest& ReferenceManifest)
+					   const FDirectoryManifest& ReferenceManifest,
+					   const FSyncFilter*		 SyncFilter)
 {
+	auto ShouldCleanup = [SyncFilter](const FPath& Filename) -> bool {
+		if (SyncFilter)
+		{
+			return SyncFilter->ShouldCleanup(Filename);
+		}
+		else
+		{
+			return true;
+		}
+	};
+
 	for (const auto& TargetManifestEntry : TargetDirectoryManifest.Files)
 	{
 		const std::wstring& TargetFileName = TargetManifestEntry.first;
 		if (ReferenceManifest.Files.find(TargetFileName) == ReferenceManifest.Files.end())
 		{
 			FPath FilePath = TargetDirectory / TargetFileName;
+
+			if (!ShouldCleanup(FilePath))
+			{
+				UNSYNC_VERBOSE2(L"Skipped deleting '%ls' (excluded by filter)", FilePath.wstring().c_str());
+				continue;
+			}
+
 			if (GDryRun)
 			{
 				UNSYNC_VERBOSE(L"Deleting '%ls' (skipped due to dry run mode)", FilePath.wstring().c_str());
@@ -3641,7 +3660,7 @@ SyncDirectory(const FSyncDirectoryOptions& SyncOptions)
 	{
 		UNSYNC_VERBOSE(L"Deleting unnecessary files");
 		UNSYNC_LOG_INDENT;
-		DeleteUnnecessaryFiles(TargetPath, TargetDirectoryManifest, SourceDirectoryManifest);
+		DeleteUnnecessaryFiles(TargetPath, TargetDirectoryManifest, SourceDirectoryManifest, SyncFilter);
 	}
 
 	// Save the source directory manifest on success.
@@ -4042,25 +4061,38 @@ ComputeNeedListSize(const FNeedList& NeedList)
 	return Result;
 }
 
-FSyncFilter::FSyncFilter(const std::wstring& InExcludedWords)
+static void
+AddCommaSeparatedWordsToList(const std::wstring& CommaSeparatedWords, std::vector<std::wstring>& Output)
 {
 	size_t Offset = 0;
-	size_t Len	  = InExcludedWords.length();
+	size_t Len	  = CommaSeparatedWords.length();
 	while (Offset < Len)
 	{
-		size_t MatchOffset = InExcludedWords.find(L',', Offset);
+		size_t MatchOffset = CommaSeparatedWords.find(L',', Offset);
 		if (MatchOffset == std::wstring::npos)
 		{
 			MatchOffset = Len;
 		}
 
-		std::wstring Word = InExcludedWords.substr(Offset, MatchOffset - Offset);
-		ExcludedWords.push_back(Word);
+		std::wstring Word = CommaSeparatedWords.substr(Offset, MatchOffset - Offset);
+		Output.push_back(Word);
 
 		Offset = MatchOffset + 1;
 	}
 }
 
+void
+FSyncFilter::ExcludeFromSync(const std::wstring& CommaSeparatedWords)
+{
+	AddCommaSeparatedWordsToList(CommaSeparatedWords, SyncExcludedWords);
+}
+
+void
+FSyncFilter::ExcludeFromCleanup(const std::wstring& CommaSeparatedWords)
+{
+	AddCommaSeparatedWordsToList(CommaSeparatedWords, CleanupExcludedWords);
+}
+
 bool
 FSyncFilter::ShouldSync(const FPath& Filename) const
 {
@@ -4074,7 +4106,31 @@ FSyncFilter::ShouldSync(const FPath& Filename) const
 bool
 FSyncFilter::ShouldSync(const std::wstring& Filename) const
 {
-	for (const std::wstring& Word : ExcludedWords)
+	for (const std::wstring& Word : SyncExcludedWords)
+	{
+		if (Filename.find(Word) != std::wstring::npos)
+		{
+			return false;
+		}
+	}
+
+	return true;
+}
+
+bool
+FSyncFilter::ShouldCleanup(const FPath& Filename) const
+{
+#if UNSYNC_PLATFORM_WINDOWS
+	return ShouldCleanup(Filename.native());
+#else
+	return ShouldCleanup(Filename.wstring());
+#endif
+}
+
+bool
+FSyncFilter::ShouldCleanup(const std::wstring& Filename) const
+{
+	for (const std::wstring& Word : CleanupExcludedWords)
 	{
 		if (Filename.find(Word) != std::wstring::npos)
 		{
diff --git a/Engine/Source/Programs/Unsync/Private/UnsyncCore.h b/Engine/Source/Programs/Unsync/Private/UnsyncCore.h
--- a/Engine/Source/Programs/Unsync/Private/UnsyncCore.h
+++ b/Engine/Source/Programs/Unsync/Private/UnsyncCore.h
@@ -322,13 +322,22 @@ FFileSyncResult SyncFile(const FPath&			 SourceFilePath,
 
 struct FSyncFilter
 {
-	FSyncFilter(const std::wstring& ExcludedWords);
+	FSyncFilter() = default;
+
+	void ExcludeFromSync(const std::wstring& CommaSeparatedWords);
+	void ExcludeFromCleanup(const std::wstring& CommaSeparatedWords);
 
 	bool  ShouldSync(const FPath& Filename) const;
 	bool  ShouldSync(const std::wstring& Filename) const;
+
+	bool  ShouldCleanup(const FPath& Filename) const;
+	bool  ShouldCleanup(const std::wstring& Filename) const;
+
 	FPath Resolve(const FPath& Filename) const;
 
-	std::vector<std::wstring> ExcludedWords;
+	std::vector<std::wstring> SyncExcludedWords; // any paths that contain these words will not be synced
+	std::vector<std::wstring> CleanupExcludedWords; // any paths that contain these words will not be deleted after sync
+
 	std::vector<FDfsAlias>	  DfsAliases;
 };
 
diff --git a/Engine/Source/Programs/Unsync/Private/UnsyncMain.cpp b/Engine/Source/Programs/Unsync/Private/UnsyncMain.cpp
--- a/Engine/Source/Programs/Unsync/Private/UnsyncMain.cpp
+++ b/Engine/Source/Programs/Unsync/Private/UnsyncMain.cpp
@@ -54,6 +54,7 @@ InnerMain(int Argc, char** Argv)
 	std::string				 InputFilename2Utf8;
 	std::string				 SourceManifestFilenameUtf8;
 	std::vector<std::string> ExcludeFilterArrayUtf8;
+	std::vector<std::string> CleanupExcludeFilterArrayUtf8;
 	std::vector<std::string> OverlayArrayUtf8;
 	std::string				 RemoteAddressUtf8;
 	std::string				 PreferredDfsUtf8;
@@ -171,6 +172,7 @@ InnerMain(int Argc, char** Argv)
 						HttpHeaderFilenameUtf8,
 						"Text file that contains any extra HTTP headers to pass to the remote server (auth tokens, etc.)");
 	SubSync->add_flag("--no-cleanup", bNoCleanupAfterSync, "Do not delete local files that aren't in the manifest after a successful sync");
+	SubSync->add_option("--cleanup-exclude", CleanupExcludeFilterArrayUtf8, "Exclude filenames that contain specified words from cleanup process (comma separated)");
 
 	// Deprecated --quick flag
 	SubSync
@@ -442,20 +444,23 @@ InnerMain(int Argc, char** Argv)
 	Algorithm.StrongHashAlgorithmId = DefaultStrongHasher;
 	Algorithm.WeakHashAlgorithmId	= DefaultWeakHasher;
 
-	// Merge multiple exclude filter strings into one
-	std::string ExcludeFilterUtf8;
+	FSyncFilter SyncFilter;
+
 	for (const std::string& Str : ExcludeFilterArrayUtf8)
 	{
-		ExcludeFilterUtf8 += Str + ",";
+		SyncFilter.ExcludeFromSync(ConvertUtf8ToWide(Str));
 	}
 
-	// Remove any trailing commas
-	while (!ExcludeFilterUtf8.empty() && ExcludeFilterUtf8.back() == ',')
+	if (const char* EnvCleanupExclude = getenv("UNSYNC_CLEANUP_EXCLUDE"))
 	{
-		ExcludeFilterUtf8.pop_back();
+		UNSYNC_VERBOSE(L"Using UNSYNC_CLEANUP_EXCLUDE environment: '%hs'", EnvCleanupExclude);
+		CleanupExcludeFilterArrayUtf8.push_back(EnvCleanupExclude);
 	}
 
-	FSyncFilter SyncFilter(ConvertUtf8ToWide(ExcludeFilterUtf8));
+	for (const std::string& Str : CleanupExcludeFilterArrayUtf8)
+	{
+		SyncFilter.ExcludeFromCleanup(ConvertUtf8ToWide(Str));
+	}
 
 	if (PreferredDfsUtf8.empty())
 	{
