unsync - Load trusted certificates from unsync.cer if it exists next to unsync.exe

diff --git a/Engine/Source/Programs/Unsync/Private/UnsyncMain.cpp b/Engine/Source/Programs/Unsync/Private/UnsyncMain.cpp
--- a/Engine/Source/Programs/Unsync/Private/UnsyncMain.cpp
+++ b/Engine/Source/Programs/Unsync/Private/UnsyncMain.cpp
@@ -29,6 +29,8 @@ UNSYNC_THIRD_PARTY_INCLUDES_END
 
 namespace unsync {
 
+static FPath GExePath;
+
 int
 InnerMain(int Argc, char** Argv)
 {
@@ -581,6 +583,25 @@ InnerMain(int Argc, char** Argv)
 		FBuffer CacertBuffer = ReadFileToBuffer(CacertPath);
 
 		RemoteDesc.TlsCacert = std::make_shared<FBuffer>(std::move(CacertBuffer));
+		RemoteDesc.TlsCacert->PushBack('\n');
+	}
+
+	{
+		FPath ExtraCertPath = GExePath.parent_path() / "unsync.cer";
+		FBuffer CertBuffer = ReadFileToBuffer(ExtraCertPath);
+		if (!CertBuffer.Empty())
+		{
+			UNSYNC_LOG(L"Using trusted certificates from '%ls'", ExtraCertPath.wstring().c_str());
+			if (!RemoteDesc.TlsCacert)
+			{
+				RemoteDesc.TlsCacert = std::make_shared<FBuffer>(std::move(CertBuffer));
+			}
+			else
+			{
+				RemoteDesc.TlsCacert->Append(CertBuffer);
+			}
+			RemoteDesc.TlsCacert->PushBack('\n');
+		}
 	}
 
 	if (Cli.got_subcommand(SubHash))
@@ -743,6 +764,7 @@ main(int argc, char** argv)
 	{
 		ArgvStringsUtf8.push_back(ConvertWideToUtf8(ArgvWide[I]));
 	}
+	GExePath = FPath(ArgvWide[0]);
 	LocalFree(ArgvWide);
 
 	std::vector<char*> ArgvUtf8;
@@ -751,8 +773,12 @@ main(int argc, char** argv)
 	{
 		ArgvUtf8.push_back(ArgvStringsUtf8[I].data());
 	}
+#else // UNSYNC_PLATFORM_WINDOWS
+	GExePath = FPath(argv[0]);
 #endif	// UNSYNC_PLATFORM_WINDOWS
 
+	GExePath = std::filesystem::weakly_canonical(GExePath);
+
 #if UNSYNC_PLATFORM_UNIX
 	std::vector<char*> ArgvUtf8;
 	ArgvUtf8.reserve(argc);
