Add --include argument to sync subcommand to limit the files synced
Add --files flag to info subcommand to list files in manifests

diff --git a/Engine/Source/Programs/Unsync/Private/UnsyncCmdSync.cpp b/Engine/Source/Programs/Unsync/Private/UnsyncCmdSync.cpp
--- a/Engine/Source/Programs/Unsync/Private/UnsyncCmdSync.cpp
+++ b/Engine/Source/Programs/Unsync/Private/UnsyncCmdSync.cpp
@@ -26,6 +26,23 @@ CmdSync(const FCmdSyncOptions& Options)
 		UNSYNC_VERBOSE(L"Sync source: '%ls' ('%ls')", Options.Source.wstring().c_str(), ResolvedSource.wstring().c_str());
 	}
 
+	if (!Options.Filter->SyncIncludedWords.empty())
+	{
+		UNSYNC_VERBOSE(L"Include filter: ");
+		for( const std::wstring& include : Options.Filter->SyncIncludedWords)
+		{ 
+			UNSYNC_VERBOSE(L"\t%s", include.c_str());
+		}
+	}
+	if(!Options.Filter->SyncExcludedWords.empty())
+	{
+		UNSYNC_VERBOSE(L"Exclude filter: ");
+		for( const std::wstring& exclude : Options.Filter->SyncExcludedWords)
+		{ 
+			UNSYNC_VERBOSE(L"\t%s", exclude.c_str());
+		}
+	}
+
 	const bool bSourcePathExists	 = PathExists(ResolvedSource, ErrorCode);
 	const bool bSourceIsDirectory	 = bSourcePathExists && unsync::IsDirectory(ResolvedSource);
 	const bool bSourceIsManifestHash = !bSourcePathExists && LooksLikeHash160(Options.Source.native());
diff --git a/Engine/Source/Programs/Unsync/Private/UnsyncCore.cpp b/Engine/Source/Programs/Unsync/Private/UnsyncCore.cpp
--- a/Engine/Source/Programs/Unsync/Private/UnsyncCore.cpp
+++ b/Engine/Source/Programs/Unsync/Private/UnsyncCore.cpp
@@ -4081,6 +4081,12 @@ AddCommaSeparatedWordsToList(const std::wstring& CommaSeparatedWords, std::vecto
 	}
 }
 
+void
+FSyncFilter::IncludeInSync(const std::wstring& CommaSeparatedWords)
+{
+	AddCommaSeparatedWordsToList(CommaSeparatedWords, SyncIncludedWords);
+}
+
 void
 FSyncFilter::ExcludeFromSync(const std::wstring& CommaSeparatedWords)
 {
@@ -4106,6 +4112,21 @@ FSyncFilter::ShouldSync(const FPath& Filename) const
 bool
 FSyncFilter::ShouldSync(const std::wstring& Filename) const
 {
+	bool bInclude = SyncIncludedWords.empty(); // Include everything if there are no specific inclusions
+	for (const std::wstring& Word : SyncIncludedWords)
+	{
+		if (Filename.find(Word) != std::wstring::npos)
+		{
+			bInclude = true;
+			break;
+		}
+	}
+
+	if (!bInclude)
+	{
+		return false;
+	}
+	
 	for (const std::wstring& Word : SyncExcludedWords)
 	{
 		if (Filename.find(Word) != std::wstring::npos)
@@ -4213,6 +4234,24 @@ LogManifestInfo(ELogLevel LogLevel, const FDirectoryManifest& Manifest)
 	LogManifestInfo(LogLevel, Info);
 }
 
+void
+LogManifestFiles(ELogLevel LogLevel, const FDirectoryManifest& Manifest)
+{
+	std::vector<std::wstring> Files;
+	for (const auto& ManifestIt : Manifest.Files)
+	{
+		Files.push_back(ManifestIt.first);
+	}
+
+	std::sort(Files.begin(), Files.end());
+
+	for (const std::wstring& Filename : Files)
+	{
+		const FFileManifest& Info = Manifest.Files.at(Filename);
+		LogPrintf(LogLevel, L"%s : %llu\n", Filename.c_str(), Info.Size);
+	}
+}
+
 HashMap<FGenericHash, FGenericBlock>
 BuildBlockMap(const FDirectoryManifest& Manifest, bool bNeedMacroBlocks)
 {
@@ -4276,7 +4315,7 @@ LogManifestDiff(ELogLevel LogLevel, const FDirectoryManifest& ManifestA, const F
 }
 
 int32
-CmdInfo(const FPath& InputA, const FPath& InputB)
+CmdInfo(const FPath& InputA, const FPath& InputB, bool bListFiles)
 {
 	FPath DirectoryManifestPathA = InputA / ".unsync" / "manifest.bin";
 	FPath DirectoryManifestPathB = InputB / ".unsync" / "manifest.bin";
@@ -4296,6 +4335,12 @@ CmdInfo(const FPath& InputA, const FPath& InputB)
 		LogManifestInfo(ELogLevel::Info, ManifestA);
 	}
 
+	if (bListFiles)
+	{
+		UNSYNC_LOG_INDENT;
+		LogManifestFiles(ELogLevel::Info, ManifestA);
+	}
+
 	if (InputB.empty())
 	{
 		return 0;
@@ -4316,6 +4361,11 @@ CmdInfo(const FPath& InputA, const FPath& InputB)
 		UNSYNC_LOG_INDENT;
 		LogManifestInfo(ELogLevel::Info, ManifestB);
 	}
+	if (bListFiles)
+	{
+		UNSYNC_LOG_INDENT;
+		LogManifestFiles(ELogLevel::Info, ManifestB);
+	}
 
 	LogPrintf(ELogLevel::Info, L"Difference:\n");
 
diff --git a/Engine/Source/Programs/Unsync/Private/UnsyncCore.h b/Engine/Source/Programs/Unsync/Private/UnsyncCore.h
--- a/Engine/Source/Programs/Unsync/Private/UnsyncCore.h
+++ b/Engine/Source/Programs/Unsync/Private/UnsyncCore.h
@@ -184,6 +184,7 @@ struct FDirectoryManifestInfo
 
 FDirectoryManifestInfo GetManifestInfo(const FDirectoryManifest& Manifest);
 void				   LogManifestInfo(ELogLevel LogLevel, const FDirectoryManifestInfo& Info);
+void				   LogManifestFiles(ELogLevel LogLevel, const FDirectoryManifestInfo& Info);
 
 const std::string& GetVersionString();
 
@@ -324,6 +325,8 @@ struct FSyncFilter
 {
 	FSyncFilter() = default;
 
+	// By default all files will be included, calling this will include only files containing these substrings
+	void IncludeInSync(const std::wstring& CommaSeparatedWords);
 	void ExcludeFromSync(const std::wstring& CommaSeparatedWords);
 	void ExcludeFromCleanup(const std::wstring& CommaSeparatedWords);
 
@@ -335,6 +338,7 @@ struct FSyncFilter
 
 	FPath Resolve(const FPath& Filename) const;
 
+	std::vector<std::wstring> SyncIncludedWords;
 	std::vector<std::wstring> SyncExcludedWords; // any paths that contain these words will not be synced
 	std::vector<std::wstring> CleanupExcludedWords; // any paths that contain these words will not be deleted after sync
 
@@ -385,6 +389,6 @@ FHash160 ComputeSerializedManifestHash160(const FDirectoryManifest& Manifest);
 FBlock128			   ToBlock128(const FGenericBlock& GenericBlock);
 std::vector<FBlock128> ToBlock128(FGenericBlockArray& GenericBlocks);
 
-int32 CmdInfo(const FPath& InputA, const FPath& InputB);
+int32 CmdInfo(const FPath& InputA, const FPath& InputB, bool bListFiles);
 
 }  // namespace unsync
diff --git a/Engine/Source/Programs/Unsync/Private/UnsyncMain.cpp b/Engine/Source/Programs/Unsync/Private/UnsyncMain.cpp
--- a/Engine/Source/Programs/Unsync/Private/UnsyncMain.cpp
+++ b/Engine/Source/Programs/Unsync/Private/UnsyncMain.cpp
@@ -55,6 +55,7 @@ InnerMain(int Argc, char** Argv)
 	std::string				 PatchFilenameUtf8;
 	std::string				 InputFilename2Utf8;
 	std::string				 SourceManifestFilenameUtf8;
+	std::vector<std::string> IncludeFilterArrayUtf8;
 	std::vector<std::string> ExcludeFilterArrayUtf8;
 	std::vector<std::string> CleanupExcludeFilterArrayUtf8;
 	std::vector<std::string> OverlayArrayUtf8;
@@ -77,6 +78,7 @@ InnerMain(int Argc, char** Argv)
 	bool					 bNoCleanupAfterSync = false;
 	bool					 bFullSourceScan	 = false;
 	bool					 bFullDifference	 = false;
+	bool					 bInfoFiles			 = false;
 	int32					 CompressionLevel	 = 3;
 	uint32					 DiffBlockSize		 = uint32(4_KB);
 	uint32					 HashOrSyncBlockSize = uint32(64_KB);
@@ -141,6 +143,7 @@ InnerMain(int Argc, char** Argv)
 	CLI::App* SubInfo = Cli.add_subcommand("info", "Display information about a manifest file or diff two manifests");
 	SubInfo->add_option("Input 1", InputFilenameUtf8, "Input manifest file or root directory")->required();
 	SubInfo->add_option("Input 2", InputFilename2Utf8, "Optional input manifest file or root directory");
+	SubInfo->add_flag("--files", bInfoFiles, "List all files in the manifest");
 	SubCommands.push_back(SubInfo);
 
 	CLI::App* SubDiff = Cli.add_subcommand("diff", "Compute difference required to transform BaseFile into SourceFile");
@@ -164,7 +167,8 @@ InnerMain(int Argc, char** Argv)
 						"FProxy server address ([transport://]address[:port][/request][#namespace])");
 	SubSync->add_option("--dfs", PreferredDfsUtf8, "Preferred DFS mirror (matched by sub-string)");
 	SubSync->add_option("--overlay", OverlayArrayUtf8, "Additional source directory to sync (keep unique files from all sources, overwrite conflicting files with overlay source)");
-	SubSync->add_option("--exclude", ExcludeFilterArrayUtf8, "Exclude filenames that contain specified words (comma separated)");
+	SubSync->add_option("--include", IncludeFilterArrayUtf8, "Include filenames that contain specified words (comma separated). If this is not present, all files will be included.");
+	SubSync->add_option("--exclude", ExcludeFilterArrayUtf8, "Exclude filenames that contain specified words (comma separated). Filter is run after --include.");
 #if UNSYNC_USE_TLS
 	AddTlsOptions(SubSync);
 #else
@@ -452,6 +456,10 @@ InnerMain(int Argc, char** Argv)
 	{
 		SyncFilter.ExcludeFromSync(ConvertUtf8ToWide(Str));
 	}
+	for (const std::string& Str : IncludeFilterArrayUtf8)
+	{
+		SyncFilter.IncludeInSync(ConvertUtf8ToWide(Str));
+	}
 
 	if (const char* EnvCleanupExclude = getenv("UNSYNC_CLEANUP_EXCLUDE"))
 	{
@@ -679,7 +687,7 @@ InnerMain(int Argc, char** Argv)
 	}
 	else if (Cli.got_subcommand(SubInfo))
 	{
-		return CmdInfo(InputFilename, InputFilename2);
+		return CmdInfo(InputFilename, InputFilename2, bInfoFiles);
 	}
 	else if (Cli.got_subcommand(SubQuery))
 	{
