unsync - Resolve letter-based source paths to universal form when syncing

diff --git a/Engine/Source/Programs/Unsync/Private/UnsyncCmdSync.cpp b/Engine/Source/Programs/Unsync/Private/UnsyncCmdSync.cpp
--- a/Engine/Source/Programs/Unsync/Private/UnsyncCmdSync.cpp
+++ b/Engine/Source/Programs/Unsync/Private/UnsyncCmdSync.cpp
@@ -9,8 +9,15 @@ namespace unsync {
 int32  // TODO: return a TResult
 CmdSync(const FCmdSyncOptions& Options)
 {
-	auto ResolvePath = [&Options](const FPath& Path) -> FPath
+	auto ResolvePath = [&Options](FPath Path) -> FPath
 	{
+#if UNSYNC_PLATFORM_WINDOWS
+		if (!Path.native().starts_with(L"\\\\"))
+		{
+			Path = GetUniversalPath(Path);
+		}
+#endif // UNSYNC_PLATFORM_WINDOWS
+
 		return Options.Filter ? Options.Filter->Resolve(Path) : Path;
 	};
 
@@ -49,6 +56,8 @@ CmdSync(const FCmdSyncOptions& Options)
 
 	bool bSourceFileSystemRequired = !bSourceIsManifestHash;
 
+	std::vector<FPath> ResolvedOverlays;
+
 	if (!Options.Overlays.empty())
 	{
 		for (const FPath& Entry : Options.Overlays)
@@ -62,6 +71,7 @@ CmdSync(const FCmdSyncOptions& Options)
 			{
 				UNSYNC_VERBOSE(L"Sync overlay: '%ls' ('%ls')", Entry.wstring().c_str(), ResolvedEntry.wstring().c_str());
 			}
+			ResolvedOverlays.push_back(ResolvedEntry);
 		}
 
 		bSourceFileSystemRequired = true;
@@ -117,10 +127,10 @@ CmdSync(const FCmdSyncOptions& Options)
 				SyncOptions.SourceType = ESyncSourceType::FileSystem;
 			}
 
-			SyncOptions.Source				   = Options.Source;
+			SyncOptions.Source				   = ResolvedSource;
 			SyncOptions.Base				   = Options.Target;  // read base data from existing target
 			SyncOptions.Target				   = Options.Target;
-			SyncOptions.Overlays			   = Options.Overlays;
+			SyncOptions.Overlays			   = ResolvedOverlays;
 			SyncOptions.SourceManifestOverride = Options.SourceManifestOverride;
 			SyncOptions.Remote				   = &Options.Remote;
 			SyncOptions.SyncFilter			   = Options.Filter;
diff --git a/Engine/Source/Programs/Unsync/Private/UnsyncCore.cpp b/Engine/Source/Programs/Unsync/Private/UnsyncCore.cpp
--- a/Engine/Source/Programs/Unsync/Private/UnsyncCore.cpp
+++ b/Engine/Source/Programs/Unsync/Private/UnsyncCore.cpp
@@ -25,7 +25,7 @@ UNSYNC_THIRD_PARTY_INCLUDES_START
 #include <flat_hash_map.hpp>
 UNSYNC_THIRD_PARTY_INCLUDES_END
 
-#define UNSYNC_VERSION_STR "1.0.46"
+#define UNSYNC_VERSION_STR "1.0.47"
 
 namespace unsync {
 
diff --git a/Engine/Source/Programs/Unsync/Private/UnsyncUtil.cpp b/Engine/Source/Programs/Unsync/Private/UnsyncUtil.cpp
--- a/Engine/Source/Programs/Unsync/Private/UnsyncUtil.cpp
+++ b/Engine/Source/Programs/Unsync/Private/UnsyncUtil.cpp
@@ -8,9 +8,11 @@
 #	include <lm.h>
 #	include <lmdfs.h>
 #	include <wincrypt.h>
+#	include <winnetwk.h> // for WNetGetUniversalName
 #	pragma comment(lib, "Netapi32.lib")
 #	pragma comment(lib, "Crypt32.lib")
 #	pragma comment(lib, "Bcrypt.lib")
+#	pragma comment(lib, "Mpr.lib")
 #endif	// UNSYNC_PLATFORM_WINDOWS
 
 #include <codecvt>
@@ -265,6 +267,30 @@ DfsEnumerate(const FPath& Root)
 	return Result;
 }
 
+FPath GetUniversalPath(const FPath& Path)
+{
+	FPath Result = Path;
+
+#if UNSYNC_PLATFORM_WINDOWS
+
+	// https://docs.microsoft.com/en-us/windows/win32/api/winnetwk/nf-winnetwk-wnetgetuniversalnamea
+
+	static constexpr DWORD MaxBufferSize = 1024;
+	WCHAR Buffer[MaxBufferSize];
+	DWORD BufferSize = MaxBufferSize;
+	UNIVERSAL_NAME_INFO* UniversalNameInfo = (UNIVERSAL_NAME_INFO*)Buffer;
+
+	DWORD ErrorCode = WNetGetUniversalNameW(Path.native().c_str(), UNIVERSAL_NAME_INFO_LEVEL, (LPVOID)UniversalNameInfo, &BufferSize);
+	if (ErrorCode == NO_ERROR)
+	{
+		Result = FPath(UniversalNameInfo->lpUniversalName);
+	}
+
+#endif	// UNSYNC_PLATFORM_WINDOWS
+
+	return Result;
+}
+
 FPath
 NormalizeFilenameUtf8(const std::string& InFilename)
 {
diff --git a/Engine/Source/Programs/Unsync/Private/UnsyncUtil.h b/Engine/Source/Programs/Unsync/Private/UnsyncUtil.h
--- a/Engine/Source/Programs/Unsync/Private/UnsyncUtil.h
+++ b/Engine/Source/Programs/Unsync/Private/UnsyncUtil.h
@@ -153,6 +153,10 @@ struct FDfsAlias
 	FPath Target;
 };
 
+// Takes a drive-based path (e.g. P:/Foo/Bar) and converts it to universal form (e.g. //server/Foo/Bar), if possible.
+// Otherwise, returns original path.
+FPath GetUniversalPath(const FPath& Path);
+
 template<typename T>
 struct TArrayView
 {
