unsync - Fix assert due to non-absolute path to unsync.cer in some cases

diff --git a/Engine/Source/Programs/Unsync/Private/UnsyncFile.cpp b/Engine/Source/Programs/Unsync/Private/UnsyncFile.cpp
--- a/Engine/Source/Programs/Unsync/Private/UnsyncFile.cpp
+++ b/Engine/Source/Programs/Unsync/Private/UnsyncFile.cpp
@@ -30,7 +30,14 @@ bool GForceBufferedFiles = false;
 static FPath
 MakeExtendedAbsolutePath(const FPath& InAbsolutePath)
 {
-	UNSYNC_ASSERT(InAbsolutePath.is_absolute());
+	if (InAbsolutePath.empty())
+	{
+		return FPath();
+	}
+
+	UNSYNC_ASSERTF(InAbsolutePath.is_absolute(),
+		L"Input path '%ls' must be absolute",
+		InAbsolutePath.wstring().c_str());
 
 #if UNSYNC_PLATFORM_WINDOWS
 	const std::wstring& InFilenameString = InAbsolutePath.native();
diff --git a/Engine/Source/Programs/Unsync/Private/UnsyncMain.cpp b/Engine/Source/Programs/Unsync/Private/UnsyncMain.cpp
--- a/Engine/Source/Programs/Unsync/Private/UnsyncMain.cpp
+++ b/Engine/Source/Programs/Unsync/Private/UnsyncMain.cpp
@@ -786,6 +786,7 @@ main(int argc, char** argv)
 #endif	// UNSYNC_PLATFORM_WINDOWS
 
 	GExePath = std::filesystem::weakly_canonical(GExePath);
+	GExePath = GetAbsoluteNormalPath(GExePath);
 
 #if UNSYNC_PLATFORM_UNIX
 	std::vector<char*> ArgvUtf8;
diff --git a/Engine/Source/Programs/Unsync/Private/UnsyncUtil.cpp b/Engine/Source/Programs/Unsync/Private/UnsyncUtil.cpp
--- a/Engine/Source/Programs/Unsync/Private/UnsyncUtil.cpp
+++ b/Engine/Source/Programs/Unsync/Private/UnsyncUtil.cpp
@@ -307,6 +307,15 @@ NormalizeFilenameUtf8(const std::string& InFilename)
 	return AbsoluteNormalPath;
 }
 
+FPath
+GetAbsoluteNormalPath(const FPath& InPath)
+{
+	FPath NormalPath = InPath.lexically_normal();
+	FPath AbsoluteNormalPath = std::filesystem::absolute(NormalPath);
+	return AbsoluteNormalPath;
+}
+
+
 const FBuffer&
 GetSystemRootCerts()
 {
diff --git a/Engine/Source/Programs/Unsync/Private/UnsyncUtil.h b/Engine/Source/Programs/Unsync/Private/UnsyncUtil.h
--- a/Engine/Source/Programs/Unsync/Private/UnsyncUtil.h
+++ b/Engine/Source/Programs/Unsync/Private/UnsyncUtil.h
@@ -281,6 +281,7 @@ CountLeadingZeros64(uint64 X)
 }
 
 FPath NormalizeFilenameUtf8(const std::string& InFilename);
+FPath GetAbsoluteNormalPath(const FPath& InPath);
 
 const FBuffer& GetSystemRootCerts();
 
