unsync - Add available disk check before starting sync, add anonymized machine name to telemetry

diff --git a/Engine/Source/Programs/Unsync/Private/UnsyncCmdSync.cpp b/Engine/Source/Programs/Unsync/Private/UnsyncCmdSync.cpp
--- a/Engine/Source/Programs/Unsync/Private/UnsyncCmdSync.cpp
+++ b/Engine/Source/Programs/Unsync/Private/UnsyncCmdSync.cpp
@@ -138,6 +138,7 @@ CmdSync(const FCmdSyncOptions& Options)
 			SyncOptions.bValidateSourceFiles   = Options.bFullSourceScan;
 			SyncOptions.bFullDifference		   = Options.bFullDifference;
 			SyncOptions.bValidateTargetFiles   = Options.bValidateTargetFiles;
+			SyncOptions.bCheckAvailableSpace   = Options.bCheckAvailableSpace;
 
 			return SyncDirectory(SyncOptions) ? 0 : 1;
 		}
diff --git a/Engine/Source/Programs/Unsync/Private/UnsyncCmdSync.h b/Engine/Source/Programs/Unsync/Private/UnsyncCmdSync.h
--- a/Engine/Source/Programs/Unsync/Private/UnsyncCmdSync.h
+++ b/Engine/Source/Programs/Unsync/Private/UnsyncCmdSync.h
@@ -24,6 +24,8 @@ struct FCmdSyncOptions
 	bool bFullSourceScan = false;
 	bool bCleanup		 = false;
 
+	bool bCheckAvailableSpace = true;
+
 	bool bValidateTargetFiles = true;  // WARNING: turning this off is intended only for testing/profiling
 
 	uint32 BlockSize = uint32(64_KB);
diff --git a/Engine/Source/Programs/Unsync/Private/UnsyncCore.cpp b/Engine/Source/Programs/Unsync/Private/UnsyncCore.cpp
--- a/Engine/Source/Programs/Unsync/Private/UnsyncCore.cpp
+++ b/Engine/Source/Programs/Unsync/Private/UnsyncCore.cpp
@@ -25,7 +25,7 @@ UNSYNC_THIRD_PARTY_INCLUDES_START
 #include <flat_hash_map.hpp>
 UNSYNC_THIRD_PARTY_INCLUDES_END
 
-#define UNSYNC_VERSION_STR "1.0.47"
+#define UNSYNC_VERSION_STR "1.0.48"
 
 namespace unsync {
 
@@ -2965,6 +2965,20 @@ LoadAndMergeSourceManifest(FDirectoryManifest& Output,
 	return MergeManifests(Output, LoadedManifest, bCaseSensitiveTargetFileSystem);
 }
 
+static std::string
+GetAnonymizedHostName()
+{
+	std::string HostName = GetCurrentHostName();
+	if (HostName.empty())
+	{
+		return {};
+	}
+	HostName += " {22FF4421-8CAC-4A14-9E4C-780AAF8BBF2A}";
+	FHash128	MachineId = HashBlake3String<FHash128>(HostName);
+	std::string Result	  = HashToHexString(MachineId);
+	return Result;
+}
+
 bool  // TODO: return a TResult
 SyncDirectory(const FSyncDirectoryOptions& SyncOptions)
 {
@@ -3383,10 +3397,14 @@ SyncDirectory(const FSyncDirectoryOptions& SyncOptions)
 		HashSet<FHash128> UniqueSourceBlocks;
 		uint64			  UniqueSourceBytes = 0;
 
+		uint64 TotalSyncSizeBytes = 0;
+
 		for (FileSyncTask& Item : SyncTaskList)
 		{
 			EstimatedNeedBytesFromSource += Item.NeedBytesFromSource;
 			EstimatedNeedBytesFromBase += Item.NeedBytesFromBase;
+			TotalSyncSizeBytes += Item.TotalSizeBytes;
+
 			for (const auto& Block : Item.NeedList.Source)
 			{
 				auto InsertRes = UniqueSourceBlocks.insert(Block.Hash.ToHash128());	 // #wip-widehash
@@ -3399,6 +3417,17 @@ SyncDirectory(const FSyncDirectoryOptions& SyncOptions)
 
 		UNSYNC_VERBOSE(L"Total need from source: %.2f MB", SizeMb(EstimatedNeedBytesFromSource));
 		UNSYNC_VERBOSE(L"Total need from base: %.2f MB", SizeMb(EstimatedNeedBytesFromBase));
+
+		uint64 AvailableDiskBytes = SyncOptions.bCheckAvailableSpace ? GetAvailableDiskSpace(TargetPath) : ~0ull;
+		if (TotalSyncSizeBytes > AvailableDiskBytes)
+		{
+			UNSYNC_ERROR(
+				L"Sync requires %.0f MB (%llu bytes) of disk space, but only %.0f MB (%llu bytes) is available. "
+				L"Use --no-space-validation flag to suppress this check.",
+				SizeMb(TotalSyncSizeBytes), TotalSyncSizeBytes,
+				SizeMb(AvailableDiskBytes), AvailableDiskBytes);
+			return false;
+		}
 	}
 
 	GGlobalProgressCurrent = 0;
@@ -3686,17 +3715,18 @@ SyncDirectory(const FSyncDirectoryOptions& SyncOptions)
 	{
 		FTelemetryEventSyncComplete Event;
 
-		Event.ClientVersion	   = GetVersionString();
-		Event.Session		   = ProxyPool.GetSessionId();
-		Event.Source		   = ConvertWideToUtf8(SourcePath.wstring());
-		Event.TotalBytes	   = TotalSourceSize;
-		Event.SourceBytes	   = StatSourceBytes;
-		Event.BaseBytes		   = StatBaseBytes;
-		Event.SkippedFiles	   = StatSkipped;
-		Event.FullCopyFiles	   = StatFullCopy;
-		Event.PartialCopyFiles = StatPartialCopy;
-		Event.Elapsed		   = ElapsedSeconds;
-		Event.bSuccess		   = bSyncSucceeded;
+		Event.ClientVersion		 = GetVersionString();
+		Event.Session			 = ProxyPool.GetSessionId();
+		Event.Source			 = ConvertWideToUtf8(SourcePath.wstring());
+		Event.ClientHostNameHash = GetAnonymizedHostName();
+		Event.TotalBytes		 = TotalSourceSize;
+		Event.SourceBytes		 = StatSourceBytes;
+		Event.BaseBytes			 = StatBaseBytes;
+		Event.SkippedFiles		 = StatSkipped;
+		Event.FullCopyFiles		 = StatFullCopy;
+		Event.PartialCopyFiles	 = StatPartialCopy;
+		Event.Elapsed			 = ElapsedSeconds;
+		Event.bSuccess			 = bSyncSucceeded;
 
 		ProxyPool.SendTelemetryEvent(Event);
 	}
diff --git a/Engine/Source/Programs/Unsync/Private/UnsyncCore.h b/Engine/Source/Programs/Unsync/Private/UnsyncCore.h
--- a/Engine/Source/Programs/Unsync/Private/UnsyncCore.h
+++ b/Engine/Source/Programs/Unsync/Private/UnsyncCore.h
@@ -355,10 +355,10 @@ enum class ESyncSourceType
 struct FSyncDirectoryOptions
 {
 	ESyncSourceType	   SourceType;
-	FPath			   Source;					// remote data location
-	FPath			   Target;					// output target location
-	FPath			   Base;					// base data location, which typically is the same as sync target
-	std::vector<FPath> Overlays;				// extra source directories to overlay over primary (add extra files, replace existing files)
+	FPath			   Source;	  // remote data location
+	FPath			   Target;	  // output target location
+	FPath			   Base;	  // base data location, which typically is the same as sync target
+	std::vector<FPath> Overlays;  // extra source directories to overlay over primary (add extra files, replace existing files)
 	FPath			   SourceManifestOverride;	// force the manifest to be read from a specified file instead of source directory
 	FSyncFilter*	   SyncFilter = nullptr;	// filter callback for partial sync support
 	const FRemoteDesc* Remote	  = nullptr;	// unsync proxy server connection settings
@@ -366,6 +366,7 @@ struct FSyncDirectoryOptions
 	bool			   bValidateSourceFiles = true;	 // whether to check that all source files declared in the manifest are present/valid
 	bool			   bValidateTargetFiles = true;	 // WARNING: turning this off is intended only for testing/profiling
 	bool			   bFullDifference = true;	// whether to run full file difference algorithm, even when there is an existing manifest
+	bool			   bCheckAvailableSpace = true;	 // whether to abort the sync if target path does not have enough available space
 };
 
 bool SyncDirectory(const FSyncDirectoryOptions& SyncOptions);
diff --git a/Engine/Source/Programs/Unsync/Private/UnsyncFile.cpp b/Engine/Source/Programs/Unsync/Private/UnsyncFile.cpp
--- a/Engine/Source/Programs/Unsync/Private/UnsyncFile.cpp
+++ b/Engine/Source/Programs/Unsync/Private/UnsyncFile.cpp
@@ -638,6 +638,26 @@ ToWindowsFileTime(const std::filesystem::file_time_type& T)
 {
 	return T.time_since_epoch().count();
 }
+
+uint64
+GetAvailableDiskSpace(const FPath& Path)
+{
+	ULARGE_INTEGER AvailableBytes = {};
+	ULARGE_INTEGER TotalBytes	  = {};
+	ULARGE_INTEGER FreeBytes	  = {};
+
+	BOOL bOk = GetDiskFreeSpaceExW(Path.native().c_str(), &AvailableBytes, &TotalBytes, &FreeBytes);
+
+	if (bOk)
+	{
+		return AvailableBytes.QuadPart;
+	}
+	else
+	{
+		return ~0ull;
+	}
+}
+
 #endif	// UNSYNC_PLATFORM_WINDOWS
 
 #if UNSYNC_PLATFORM_UNIX
@@ -800,6 +820,12 @@ ToWindowsFileTime(const std::filesystem::file_time_type& FileTime)
 	return ticks;
 }
 
+uint64
+GetAvailableDiskSpace(const FPath& Path)
+{
+	return ~0ull; // TODO: query available space via statvfs()
+}
+
 #endif	// UNSYNC_PLATFORM_UNIX
 
 FBuffer
diff --git a/Engine/Source/Programs/Unsync/Private/UnsyncFile.h b/Engine/Source/Programs/Unsync/Private/UnsyncFile.h
--- a/Engine/Source/Programs/Unsync/Private/UnsyncFile.h
+++ b/Engine/Source/Programs/Unsync/Private/UnsyncFile.h
@@ -453,6 +453,10 @@ bool		   FileCopy(const FPath& From, const FPath& To, std::error_code& OutErrorC
 bool		   FileCopyOverwrite(const FPath& From, const FPath& To, std::error_code& OutErrorCode);
 bool		   FileRemove(const FPath& Path, std::error_code& OutErrorCode);
 
+// Returns number of bytes that can be written to the given path.
+// Returns ~0ull if the available space could not be determined.
+uint64 GetAvailableDiskSpace(const FPath& Path);
+
 std::filesystem::recursive_directory_iterator RecursiveDirectoryScan(const FPath& Path);
 
 uint64 ToWindowsFileTime(const std::filesystem::file_time_type& T);
diff --git a/Engine/Source/Programs/Unsync/Private/UnsyncMain.cpp b/Engine/Source/Programs/Unsync/Private/UnsyncMain.cpp
--- a/Engine/Source/Programs/Unsync/Private/UnsyncMain.cpp
+++ b/Engine/Source/Programs/Unsync/Private/UnsyncMain.cpp
@@ -76,6 +76,7 @@ InnerMain(int Argc, char** Argv)
 	bool					 bIncrementalMode	 = false;
 	bool					 bNoOutputValidation = false;
 	bool					 bNoCleanupAfterSync = false;
+	bool					 bNoSpaceValidation	 = false;
 	bool					 bFullSourceScan	 = false;
 	bool					 bFullDifference	 = false;
 	bool					 bInfoFiles			 = false;
@@ -214,6 +215,7 @@ InnerMain(int Argc, char** Argv)
 					  "sync process, however this can be very slow when dealing with large numbers of files and directories.");
 
 	SubSync->add_flag("--no-output-validation", bNoOutputValidation, "Skip final patched file block hash validation (DANGEROUS)");
+	SubSync->add_flag("--no-space-validation", bNoSpaceValidation, "Skip checking available disk space before sync (DANGEROUS)");
 	SubSync->add_option("-b, --block", HashOrSyncBlockSize, "Block size in bytes (default=64KB)");
 	SubCommands.push_back(SubSync);
 
@@ -654,6 +656,7 @@ InnerMain(int Argc, char** Argv)
 		SyncOptions.BlockSize			   = HashOrSyncBlockSize;
 		SyncOptions.Filter				   = &SyncFilter;
 		SyncOptions.bValidateTargetFiles   = !bNoOutputValidation;
+		SyncOptions.bCheckAvailableSpace   = !bNoSpaceValidation;
 
 		for (const std::string& Entry : OverlayArrayUtf8)
 		{
diff --git a/Engine/Source/Programs/Unsync/Private/UnsyncProxy.cpp b/Engine/Source/Programs/Unsync/Private/UnsyncProxy.cpp
--- a/Engine/Source/Programs/Unsync/Private/UnsyncProxy.cpp
+++ b/Engine/Source/Programs/Unsync/Private/UnsyncProxy.cpp
@@ -440,7 +440,11 @@ FUnsyncProtocolImpl::SendTelemetryEvent(const FRemoteDesc& RemoteDesc, const FTe
 
 	Obj["type"]			  = "sync_complete";
 	Obj["client_version"] = Event.ClientVersion;
-	Obj["source"]		  = Event.Source;
+	if (!Event.ClientHostNameHash.empty())
+	{
+		Obj["client_host_name_hash"] = Event.ClientHostNameHash;
+	}
+	Obj["source"]	 = Event.Source;
 	Obj["total_mb"]	 = SizeMb(Event.TotalBytes);  // use size in megabytes due lack of JSON 64 bit int support and better human readability
 	Obj["source_mb"] = SizeMb(Event.SourceBytes);
 	Obj["base_mb"]	 = SizeMb(Event.BaseBytes);
diff --git a/Engine/Source/Programs/Unsync/Private/UnsyncProxy.h b/Engine/Source/Programs/Unsync/Private/UnsyncProxy.h
--- a/Engine/Source/Programs/Unsync/Private/UnsyncProxy.h
+++ b/Engine/Source/Programs/Unsync/Private/UnsyncProxy.h
@@ -84,6 +84,7 @@ struct FTelemetryEventSyncComplete
 	std::string ClientVersion;
 	std::string Session;
 	std::string Source;
+	std::string ClientHostNameHash;
 	uint64 TotalBytes = 0;
 	uint64 SourceBytes = 0;
 	uint64 BaseBytes = 0;
diff --git a/Engine/Source/Programs/Unsync/Private/UnsyncSocket.cpp b/Engine/Source/Programs/Unsync/Private/UnsyncSocket.cpp
--- a/Engine/Source/Programs/Unsync/Private/UnsyncSocket.cpp
+++ b/Engine/Source/Programs/Unsync/Private/UnsyncSocket.cpp
@@ -72,6 +72,23 @@ LazyInitSockets()
 	static FSocketInitHelper InitHelper;
 }
 
+std::string
+GetCurrentHostName()
+{
+	LazyInitSockets();
+
+	char Buffer[1024] = {};
+
+	if (gethostname(Buffer, (int)sizeof(Buffer)) == 0)
+	{
+		return std::string(Buffer);
+	}
+	else
+	{
+		return {};
+	}
+}
+
 FSocketHandle
 SocketConnectTcp(const char* DestAddress, uint16 Port)
 {
diff --git a/Engine/Source/Programs/Unsync/Private/UnsyncSocket.h b/Engine/Source/Programs/Unsync/Private/UnsyncSocket.h
--- a/Engine/Source/Programs/Unsync/Private/UnsyncSocket.h
+++ b/Engine/Source/Programs/Unsync/Private/UnsyncSocket.h
@@ -34,6 +34,9 @@ struct FSocketAddress  // opaque buffer that holds socket address
 	bool operator==(const FSocketAddress& Other) const { return Data64[0] == Other.Data64[0] && Data64[1] == Other.Data64[1]; }
 };
 
+// Returns current host name or empty string if it could not be determined.
+std::string GetCurrentHostName();
+
 FSocketHandle SocketConnectTcp(const char* DestAddress, uint16 Port);
 void		  SocketClose(FSocketHandle Socket);
 
