unsync - Workaround for ICE when using MSVC toolchain 14.34.31942

diff --git a/Engine/Source/Programs/Unsync/Private/UnsyncCore.cpp b/Engine/Source/Programs/Unsync/Private/UnsyncCore.cpp
--- a/Engine/Source/Programs/Unsync/Private/UnsyncCore.cpp
+++ b/Engine/Source/Programs/Unsync/Private/UnsyncCore.cpp
@@ -242,14 +242,14 @@ ComputeBlocksVariableT(FIOReader& Reader, uint32 BlockSize, EStrongHashAlgorithm
 	const uint32 BlocksPerMacroBlock	 = CheckedNarrow(DivUp(TargetMacroBlockSize - MinimumMacroBlockSize, BlockSize));
 	const uint32 MacroBlockHashThreshold = BlocksPerMacroBlock ? (0xFFFFFFFF / BlocksPerMacroBlock) : 0;
 
-	struct Task
+	struct FTask
 	{
 		uint64			   Offset = 0;
 		FGenericBlockArray Blocks;
 		FGenericBlockArray MacroBlocks;
 	};
 
-	std::vector<Task> Tasks;
+	std::vector<FTask> Tasks;
 	Tasks.resize(NumTasks);
 
 	FSemaphore IoSemaphore(MAX_ACTIVE_READERS);
@@ -290,7 +290,7 @@ ComputeBlocksVariableT(FIOReader& Reader, uint32 BlockSize, EStrongHashAlgorithm
 						 MinimumMacroBlockSize,
 						 MaximumMacroBlockSize,
 						 MacroBlockHashThreshold]() {
-			Task& Task = Tasks[TaskIndex];
+			FTask& Task = Tasks[TaskIndex];
 
 			const uint8* DataBegin	  = ScanTaskBuffer->Data();
 			const uint8* DataEnd	  = DataBegin + ThisTaskSize;
@@ -388,7 +388,7 @@ ComputeBlocksVariableT(FIOReader& Reader, uint32 BlockSize, EStrongHashAlgorithm
 	FGenericBlockArray Result;
 	for (uint64 I = 0; I < NumTasks; ++I)
 	{
-		const Task& Task = Tasks[I];
+		const FTask& Task = Tasks[I];
 		for (uint64 J = 0; J < Task.Blocks.size(); ++J)
 		{
 			Result.push_back(Task.Blocks[J]);
@@ -399,7 +399,7 @@ ComputeBlocksVariableT(FIOReader& Reader, uint32 BlockSize, EStrongHashAlgorithm
 	{
 		for (uint64 I = 0; I < NumTasks; ++I)
 		{
-			const Task& Task = Tasks[I];
+			const FTask& Task = Tasks[I];
 			for (uint64 J = 0; J < Task.MacroBlocks.size(); ++J)
 			{
 				OutMacroBlocks->Output.push_back(Task.MacroBlocks[J]);
