unsync - Fix --full-diff option, bump version to 1.0.49

diff --git a/Engine/Source/Programs/Unsync/Private/UnsyncCore.cpp b/Engine/Source/Programs/Unsync/Private/UnsyncCore.cpp
--- a/Engine/Source/Programs/Unsync/Private/UnsyncCore.cpp
+++ b/Engine/Source/Programs/Unsync/Private/UnsyncCore.cpp
@@ -25,7 +25,7 @@ UNSYNC_THIRD_PARTY_INCLUDES_START
 #include <flat_hash_map.hpp>
 UNSYNC_THIRD_PARTY_INCLUDES_END
 
-#define UNSYNC_VERSION_STR "1.0.48"
+#define UNSYNC_VERSION_STR "1.0.49"
 
 namespace unsync {
 
@@ -3207,6 +3207,7 @@ SyncDirectory(const FSyncDirectoryOptions& SyncOptions)
 
 		TotalSourceSize += SourceFileManifest.Size;
 
+		bool bTargetFileAttributesMatch = false;
 		auto TargetManifestIt = TargetDirectoryManifest.Files.find(SourceFilename);
 		if (TargetManifestIt != TargetDirectoryManifest.Files.end())
 		{
@@ -3214,12 +3215,17 @@ SyncDirectory(const FSyncDirectoryOptions& SyncOptions)
 
 			if (SourceFileManifest.Size == TargetFileManifest.Size && SourceFileManifest.Mtime == TargetFileManifest.Mtime)
 			{
-				StatSkipped++;
-				UNSYNC_VERBOSE2(L"Skipped '%ls' (up to date)", SourceManifestIt.first.c_str());
-				continue;
+				bTargetFileAttributesMatch = true;
 			}
 		}
 
+		if (bTargetFileAttributesMatch && !SyncOptions.bFullDifference)
+		{
+			UNSYNC_VERBOSE2(L"Skipped '%ls' (up to date)", SourceManifestIt.first.c_str());
+			StatSkipped++;
+			continue;
+		}
+
 		FPath SourceFilePath = SourceManifestIt.second.CurrentPath;
 		FPath BaseFilePath	 = BasePath / ToPath(SourceManifestIt.first);
 		FPath TargetFilePath = TargetPath / ToPath(SourceManifestIt.first);
@@ -3269,7 +3275,15 @@ SyncDirectory(const FSyncDirectoryOptions& SyncOptions)
 			}
 			else
 			{
-				UNSYNC_VERBOSE2(L"Dirty file: '%ls'", SourceManifestIt.first.c_str());
+				if (bTargetFileAttributesMatch && SyncOptions.bFullDifference)
+				{
+					UNSYNC_VERBOSE2(L"Dirty file: '%ls' (forced by --full-diff)", SourceManifestIt.first.c_str());
+				}
+				else
+				{
+					UNSYNC_VERBOSE2(L"Dirty file: '%ls'", SourceManifestIt.first.c_str());
+				}
+				
 				StatPartialCopy++;
 
 				if (bFileSystemSource && SyncOptions.bValidateSourceFiles && !SourceAttribCache.Exists(ResolvedSourceFilePath) &&
