unsync - Files that are excluded from sync via --exclude [string] argument will now be deleted

* This removes any leftover files from a previous sync that may have used a different exclusion filter
* Stale files in some cases may cause errors (for example, serialization scheme may change and new executables may no longer be compatible with old data)
* If necessary, --cleanup-exclude [string] may still be used to explicitly prevent deletion of specific stale files
* Improved `info` command (which displays basic information about a manifest file) to detect if input path is a directory or a manifest file itself

diff --git a/Engine/Source/Programs/Unsync/Private/UnsyncCore.cpp b/Engine/Source/Programs/Unsync/Private/UnsyncCore.cpp
--- a/Engine/Source/Programs/Unsync/Private/UnsyncCore.cpp
+++ b/Engine/Source/Programs/Unsync/Private/UnsyncCore.cpp
@@ -25,7 +25,7 @@ UNSYNC_THIRD_PARTY_INCLUDES_START
 #include <flat_hash_map.hpp>
 UNSYNC_THIRD_PARTY_INCLUDES_END
 
-#define UNSYNC_VERSION_STR "1.0.49"
+#define UNSYNC_VERSION_STR "1.0.50-dev"
 
 namespace unsync {
 
@@ -2859,23 +2859,24 @@ DeleteUnnecessaryFiles(const FPath&				 TargetDirectory,
 	for (const auto& TargetManifestEntry : TargetDirectoryManifest.Files)
 	{
 		const std::wstring& TargetFileName = TargetManifestEntry.first;
-		if (ReferenceManifest.Files.find(TargetFileName) == ReferenceManifest.Files.end())
+
+		auto Cleanup = [&ShouldCleanup, &TargetDirectory](const std::wstring& TargetFileName, const wchar_t* Reason) 
 		{
 			FPath FilePath = TargetDirectory / TargetFileName;
 
-			if (!ShouldCleanup(FilePath))
+			if (!ShouldCleanup(TargetFileName))
 			{
-				UNSYNC_VERBOSE2(L"Skipped deleting '%ls' (excluded by filter)", FilePath.wstring().c_str());
-				continue;
+				UNSYNC_VERBOSE2(L"Skipped deleting '%ls' (excluded by cleanup filter)", FilePath.wstring().c_str());
+				return;
 			}
 
 			if (GDryRun)
 			{
-				UNSYNC_VERBOSE(L"Deleting '%ls' (skipped due to dry run mode)", FilePath.wstring().c_str());
+				UNSYNC_VERBOSE(L"Deleting '%ls' (%ls, skipped due to dry run mode)", FilePath.wstring().c_str(), Reason);
 			}
 			else
 			{
-				UNSYNC_VERBOSE(L"Deleting '%ls'", FilePath.wstring().c_str());
+				UNSYNC_VERBOSE(L"Deleting '%ls' (%ls)", FilePath.wstring().c_str(), Reason);
 				std::error_code ErrorCode = {};
 				FileRemove(FilePath, ErrorCode);
 				if (ErrorCode)
@@ -2883,6 +2884,15 @@ DeleteUnnecessaryFiles(const FPath&				 TargetDirectory,
 					UNSYNC_VERBOSE(L"System error code %d: %hs", ErrorCode.value(), ErrorCode.message().c_str());
 				}
 			}
+		};
+
+		if (ReferenceManifest.Files.find(TargetFileName) == ReferenceManifest.Files.end())
+		{
+			Cleanup(TargetFileName, L"not in manifest");
+		}
+		else if (!SyncFilter->ShouldSync(TargetFileName))
+		{
+			Cleanup(TargetFileName, L"excluded from sync");
 		}
 	}
 }
@@ -4266,9 +4276,10 @@ LogManifestInfo(ELogLevel LogLevel, const FDirectoryManifestInfo& Info)
 	LogPrintf(LogLevel, L"Files: %llu\n", llu(Info.NumFiles));
 	LogPrintf(LogLevel, L"Blocks: %llu\n", llu(Info.NumBlocks));
 	LogPrintf(LogLevel, L"Macro blocks: %llu\n", llu(Info.NumMacroBlocks));
-	LogPrintf(LogLevel, L"Total data size: %.2f MB (%llu bytes)\n", SizeMb(Info.TotalSize), llu(Info.TotalSize));
+	LogPrintf(LogLevel, L"Total data size: %.0f MB (%llu bytes)\n", SizeMb(Info.TotalSize), llu(Info.TotalSize));
 
 	// TODO: block size distribution histogram
+	// TODO: size distribution per file extension
 }
 
 void
@@ -4361,11 +4372,12 @@ LogManifestDiff(ELogLevel LogLevel, const FDirectoryManifest& ManifestA, const F
 int32
 CmdInfo(const FPath& InputA, const FPath& InputB, bool bListFiles)
 {
-	FPath DirectoryManifestPathA = InputA / ".unsync" / "manifest.bin";
-	FPath DirectoryManifestPathB = InputB / ".unsync" / "manifest.bin";
+	FPath DirectoryManifestPathA = IsDirectory(InputA) ? (InputA / ".unsync" / "manifest.bin") : InputA;
+	FPath DirectoryManifestPathB = IsDirectory(InputB) ? (InputB / ".unsync" / "manifest.bin") : InputB;
 
 	FDirectoryManifest ManifestA;
-	bool			   bManifestAValid = LoadDirectoryManifest(ManifestA, InputA, DirectoryManifestPathA);
+
+	bool bManifestAValid = LoadDirectoryManifest(ManifestA, InputA, DirectoryManifestPathA);
 
 	if (!bManifestAValid)
 	{
@@ -4393,7 +4405,9 @@ CmdInfo(const FPath& InputA, const FPath& InputB, bool bListFiles)
 	LogPrintf(ELogLevel::Info, L"\n");
 
 	FDirectoryManifest ManifestB;
-	bool			   bManifestBValid = LoadDirectoryManifest(ManifestB, InputB, DirectoryManifestPathB);
+
+	bool bManifestBValid = LoadDirectoryManifest(ManifestB, InputB, DirectoryManifestPathB);
+
 	if (!bManifestBValid)
 	{
 		return 1;
