unsync - Bump version to 1.0.50

diff --git a/Engine/Source/Programs/Unsync/Private/UnsyncCore.cpp b/Engine/Source/Programs/Unsync/Private/UnsyncCore.cpp
--- a/Engine/Source/Programs/Unsync/Private/UnsyncCore.cpp
+++ b/Engine/Source/Programs/Unsync/Private/UnsyncCore.cpp
@@ -25,7 +25,7 @@ UNSYNC_THIRD_PARTY_INCLUDES_START
 #include <flat_hash_map.hpp>
 UNSYNC_THIRD_PARTY_INCLUDES_END
 
-#define UNSYNC_VERSION_STR "1.0.50-dev"
+#define UNSYNC_VERSION_STR "1.0.50"
 
 namespace unsync {
 
@@ -2492,15 +2492,6 @@ SyncFile(const FNeedList&		   NeedList,
 				Result.SystemErrorCode = Ec;
 			}
 		}
-
-		if (Result.Succeeded())
-		{
-			LogStatus(TargetFilePath.wstring().c_str(), L"Succeeded");
-		}
-		else
-		{
-			LogStatus(TargetFilePath.wstring().c_str(), L"Failed");
-		}
 	}
 	else
 	{
@@ -3561,6 +3552,8 @@ SyncDirectory(const FSyncDirectoryOptions& SyncOptions)
 			FFileSyncResult SyncResult =
 				SyncFile(Item.NeedList, Item.ResolvedSourceFilePath, SourceBlocks, *BaseFile.get(), Item.TargetFilePath, SyncFileOptions);
 
+			LogStatus(Item.TargetFilePath.wstring().c_str(), SyncResult.Succeeded() ? L"Succeeded" : L"Failed");
+
 			StatSourceBytes += SyncResult.SourceBytes;
 			StatBaseBytes += SyncResult.BaseBytes;
 			UNSYNC_ASSERT(SyncResult.SourceBytes + SyncResult.BaseBytes == Item.TotalSizeBytes);
