unsync - Download small files in batches to reduce amount of network requests

* This improves download performance of data sets with large numbers of small files due to reduced impact of latency
* Small files are grouped into ~4MB batches

diff --git a/Engine/Source/Programs/Unsync/Private/UnsyncBuffer.h b/Engine/Source/Programs/Unsync/Private/UnsyncBuffer.h
--- a/Engine/Source/Programs/Unsync/Private/UnsyncBuffer.h
+++ b/Engine/Source/Programs/Unsync/Private/UnsyncBuffer.h
@@ -69,6 +69,18 @@ public:
 	FBufferView	   View() const { return FBufferView{.Data = Data(), .Size = Size()}; }
 	FMutBufferView MutView() { return FMutBufferView{.Data = Data(), .Size = Size()}; }
 
+	FBufferView View(uint64 ViewOffset, uint64 ViewSize) const
+	{
+		UNSYNC_ASSERT(ViewOffset + ViewSize <= Size());
+		return FBufferView{.Data = Data() + ViewOffset, .Size = ViewSize};
+	}
+
+	FMutBufferView MutView(uint64 ViewOffset, uint64 ViewSize)
+	{
+		UNSYNC_ASSERT(ViewOffset + ViewSize <= Size());
+		return FMutBufferView{.Data = Data() + ViewOffset, .Size = ViewSize};
+	}
+
 	operator FBufferView() const { return View(); }
 
 private:
diff --git a/Engine/Source/Programs/Unsync/Private/UnsyncCore.cpp b/Engine/Source/Programs/Unsync/Private/UnsyncCore.cpp
--- a/Engine/Source/Programs/Unsync/Private/UnsyncCore.cpp
+++ b/Engine/Source/Programs/Unsync/Private/UnsyncCore.cpp
@@ -11,13 +11,13 @@
 #include "UnsyncThread.h"
 #include "UnsyncUtil.h"
 
+#include <condition_variable>
 #include <deque>
 #include <filesystem>
 #include <mutex>
 #include <string>
 #include <unordered_map>
 #include <unordered_set>
-#include <condition_variable>
 
 UNSYNC_THIRD_PARTY_INCLUDES_START
 #include <blake3.h>
@@ -25,7 +25,7 @@ UNSYNC_THIRD_PARTY_INCLUDES_START
 #include <flat_hash_map.hpp>
 UNSYNC_THIRD_PARTY_INCLUDES_END
 
-#define UNSYNC_VERSION_STR "1.0.50"
+#define UNSYNC_VERSION_STR "1.0.51-dev"
 
 namespace unsync {
 
@@ -66,8 +66,7 @@ LogGlobalProgress()
 	LogProgress(nullptr, GGlobalProgressCurrent, GGlobalProgressTotal);
 }
 
-enum class ELogProgressUnits
-{
+enum class ELogProgressUnits {
 	Raw,
 	Bytes,
 	MB,
@@ -1157,23 +1156,30 @@ BuildReadSchedule(const std::vector<FNeedBlock>& Blocks)
 	return Result;
 }
 
+class FBlockCache
+{
+public:
+	FBuffer						   BlockData;
+	HashMap<FHash128, FBufferView> BlockMap;  // Decompressed block data by hash
+};
+
 FBuffer
 BuildTargetBuffer(FIOReader&			 SourceProvider,
 				  FIOReader&			 BaseProvider,
 				  const FNeedList&		 NeedList,
 				  EStrongHashAlgorithmID StrongHasher,
-				  FProxyPool*			 ProxyPool)
+				  FProxyPool*			 ProxyPool,
+				  FBlockCache*			 BlockCache)
 {
 	FBuffer				Result;
 	const FNeedListSize SizeInfo = ComputeNeedListSize(NeedList);
 	Result.Resize(SizeInfo.TotalBytes);
 	FMemReaderWriter ResultWriter(Result.Data(), Result.Size());
-	BuildTarget(ResultWriter, SourceProvider, BaseProvider, NeedList, StrongHasher, ProxyPool);
+	BuildTarget(ResultWriter, SourceProvider, BaseProvider, NeedList, StrongHasher, ProxyPool, BlockCache);
 	return Result;
 }
 
-enum EListType
-{
+enum EListType {
 	Source,
 	Base
 };
@@ -1227,14 +1233,14 @@ DownloadBlocks(FProxyPool&					  ProxyPool,
 			SchedulerSleep(1000);
 		}
 
-		struct DownloadBatch
+		struct FDownloadBatch
 		{
 			uint64 Begin;
 			uint64 End;
 			uint64 SizeBytes;
 		};
-		std::vector<DownloadBatch> Batches;
-		Batches.push_back(DownloadBatch{});
+		std::vector<FDownloadBatch> Batches;
+		Batches.push_back(FDownloadBatch{});
 
 		const uint64 MaxBytesPerBatch = ProxyPool.RemoteDesc.Protocol == EProtocolFlavor::Jupiter ? 16_MB : 128_MB;
 
@@ -1249,10 +1255,10 @@ DownloadBlocks(FProxyPool&					  ProxyPool,
 			else
 			{
 				UNSYNC_ASSERT(Batches.back().End == I);
-				DownloadBatch NewBatch = {};
-				NewBatch.Begin		   = I;
-				NewBatch.End		   = I + 1;
-				NewBatch.SizeBytes	   = 0;
+				FDownloadBatch NewBatch = {};
+				NewBatch.Begin			= I;
+				NewBatch.End			= I + 1;
+				NewBatch.SizeBytes		= 0;
 				Batches.push_back(NewBatch);
 			}
 
@@ -1264,7 +1270,7 @@ DownloadBlocks(FProxyPool&					  ProxyPool,
 		FTaskGroup DownloadTasks;
 		std::mutex DownloadedBlocksMutex;
 
-		for (DownloadBatch Batch : Batches)
+		for (FDownloadBatch Batch : Batches)
 		{
 			ProxyPool.ParallelDownloadSemaphore.Acquire();
 			DownloadTasks.run([NeedBlocks,
@@ -1343,7 +1349,8 @@ BuildTarget(FIOWriter&			   Result,
 			FIOReader&			   Base,
 			const FNeedList&	   NeedList,
 			EStrongHashAlgorithmID StrongHasher,
-			FProxyPool*			   ProxyPool)
+			FProxyPool*			   ProxyPool,
+			FBlockCache*		   BlockCache)
 {
 	UNSYNC_LOG_INDENT;
 
@@ -1365,6 +1372,13 @@ BuildTarget(FIOWriter&			   Result,
 	const bool	 bAllowVerboseLog = GLogVerbose;
 	const uint32 LogIndent		  = GLogIndent;
 
+	// Copy the need list as it will be progressively filtered by various sources:
+	// - block cache
+	// - download from proxy
+	// - read from source file
+	// TODO: could use a bit array to indicate if the entry is valid instead of full copy
+	std::vector<FNeedBlock> FilteredSourceNeedList = NeedList.Source;
+
 	auto ProcessNeedList = [bAllowVerboseLog, LogIndent, &Result, &bGotError, &WriteSemaphore, &WriteTasks, &bWaitingForBaseData](
 							   FIOReader&					  DataProvider,
 							   const std::vector<FNeedBlock>& NeedBlocks,
@@ -1505,23 +1519,54 @@ BuildTarget(FIOWriter&			   Result,
 		}
 	});
 
-	if (ProxyPool && ProxyPool->IsValid())
+	HashSet<FHash128> CachedBlocks;
+
+	if (BlockCache)
+	{
+		UNSYNC_VERBOSE(L"Writing blocks from cache");
+		UNSYNC_LOG_INDENT;
+
+		for (const FNeedBlock NeedBlock : FilteredSourceNeedList)
+		{
+			auto It = BlockCache->BlockMap.find(NeedBlock.Hash.ToHash128());
+			if (It != BlockCache->BlockMap.end())
+			{
+				FBufferView BlockBuffer = It->second;
+				Result.Write(BlockBuffer.Data, NeedBlock.TargetOffset, BlockBuffer.Size);
+
+				AddGlobalProgress(NeedBlock.Size, EListType::Source);
+			}
+		}
+
+		auto FilterResult =
+			std::remove_if(FilteredSourceNeedList.begin(), FilteredSourceNeedList.end(), [BlockCache](const FNeedBlock& Block) {
+				return BlockCache->BlockMap.find(Block.Hash.ToHash128()) != BlockCache->BlockMap.end();
+			});
+
+		FilteredSourceNeedList.erase(FilterResult, FilteredSourceNeedList.end());
+	}
+
+	if (FilteredSourceNeedList.size() == 0)
+	{
+		// No more data is needed from source file
+	}
+	else if (ProxyPool && ProxyPool->IsValid())
 	{
-		struct BlockWriteCmd
+		struct FBlockWriteCmd
 		{
 			uint64 TargetOffset = 0;
 			uint64 Size			= 0;
 		};
 		std::mutex							   DownloadedBlocksMutex;
 		HashSet<FHash128>					   DownloadedBlocks;
-		HashMap<FHash128, BlockWriteCmd>	   BlockWriteMap;
+		HashMap<FHash128, FBlockWriteCmd>	   BlockWriteMap;
 		HashMap<FHash128, std::vector<uint64>> BlockScatterMap;
 		std::vector<FNeedBlock>				   UniqueNeedList;
 
 		uint64 EstimatedDownloadSize = 0;
-		for (const FNeedBlock& Block : NeedList.Source)
+		for (const FNeedBlock& Block : FilteredSourceNeedList)
 		{
-			BlockWriteCmd Cmd;
+			FBlockWriteCmd Cmd;
 			Cmd.TargetOffset  = Block.TargetOffset;
 			Cmd.Size		  = Block.Size;
 			auto InsertResult = BlockWriteMap.insert(std::make_pair(Block.Hash.ToHash128(), Cmd));	// #wip-widehash
@@ -1579,7 +1624,7 @@ BuildTarget(FIOWriter&			   Result,
 				auto WriteIt = BlockWriteMap.find(BlockHash);
 				if (WriteIt != BlockWriteMap.end())
 				{
-					BlockWriteCmd Cmd = WriteIt->second;
+					FBlockWriteCmd Cmd = WriteIt->second;
 
 					// TODO: avoid this copy by storing IOBuffer in DownloadedBlock
 					bool	  bCompressed	 = Block.IsCompressed();
@@ -1659,7 +1704,8 @@ BuildTarget(FIOWriter&			   Result,
 							auto ScatterIt = BlockScatterMap.find(BlockHash);
 							if (ScatterIt != BlockScatterMap.end())
 							{
-								for (uint64 ScatterOffset : ScatterIt->second)
+								const std::vector<uint64>& ScatterList = ScatterIt->second;
+								for (uint64 ScatterOffset : ScatterList)
 								{
 									uint64 WrittenBytes = Result.Write(DecompressedData.GetData(), ScatterOffset, Cmd.Size);
 									if (WrittenBytes != Cmd.Size)
@@ -1668,6 +1714,8 @@ BuildTarget(FIOWriter&			   Result,
 										bGotError = true;
 									}
 								}
+
+								AddGlobalProgress(Cmd.Size * ScatterList.size(), EListType::Source);
 							}
 						}
 
@@ -1675,8 +1723,6 @@ BuildTarget(FIOWriter&			   Result,
 						{
 							std::lock_guard<std::mutex> LockGuard(DownloadedBlocksMutex);
 							DownloadedBlocks.insert(BlockHash);
-
-							AddGlobalProgress(Cmd.Size, EListType::Source);
 						}
 
 						DecompressionSemaphore.Release();
@@ -1694,26 +1740,25 @@ BuildTarget(FIOWriter&			   Result,
 						   NumHashMismatches.load());
 		}
 
-		std::vector<FNeedBlock> FilteredNeedList;
-		for (const FNeedBlock& Block : NeedList.Source)
 		{
-			if (DownloadedBlocks.find(Block.Hash.ToHash128()) == DownloadedBlocks.end())  // #wip-widehash
-			{
-				FilteredNeedList.push_back(Block);
-			}
+			auto FilterResult =
+				std::remove_if(FilteredSourceNeedList.begin(), FilteredSourceNeedList.end(), [&DownloadedBlocks](const FNeedBlock& Block) {
+					return DownloadedBlocks.find(Block.Hash.ToHash128()) != DownloadedBlocks.end();
+				});
+			FilteredSourceNeedList.erase(FilterResult, FilteredSourceNeedList.end());
 		}
 
-		if (!FilteredNeedList.empty())
+		if (!FilteredSourceNeedList.empty())
 		{
-			uint64 RemainingBytes = ComputeSize(FilteredNeedList);
+			uint64 RemainingBytes = ComputeSize(FilteredSourceNeedList);
 
 			UNSYNC_VERBOSE(L"Did not receive %d blocks from proxy. Must read %.2f MB from source.",
-						   FilteredNeedList.size(),
+						   FilteredSourceNeedList.size(),
 						   SizeMb(RemainingBytes));
 
 			UNSYNC_LOG_INDENT;
 
-			ProcessNeedList(Source, FilteredNeedList, RemainingBytes, EListType::Source, bSourceDataCopyTaskDone);
+			ProcessNeedList(Source, FilteredSourceNeedList, RemainingBytes, EListType::Source, bSourceDataCopyTaskDone);
 		}
 		else
 		{
@@ -1722,7 +1767,7 @@ BuildTarget(FIOWriter&			   Result,
 	}
 	else if (SizeInfo.SourceBytes)
 	{
-		ProcessNeedList(Source, NeedList.Source, SizeInfo.SourceBytes, EListType::Source, bSourceDataCopyTaskDone);
+		ProcessNeedList(Source, FilteredSourceNeedList, SizeInfo.SourceBytes, EListType::Source, bSourceDataCopyTaskDone);
 	}
 
 	if (!bBaseDataCopyTaskDone)
@@ -1788,7 +1833,7 @@ CreateDirectoryManifest(const FPath& Root, uint32 BlockSize, FAlgorithmOptions A
 	FSemaphore	 Semaphore(MaxConcurrentFiles);
 
 	std::mutex ResultMutex;
-	FPath UnsyncDirName = ".unsync";
+	FPath	   UnsyncDirName = ".unsync";
 
 	for (const std::filesystem::directory_entry& Dir : RecursiveDirectoryScan(Root))
 	{
@@ -2040,11 +2085,11 @@ UpdateDirectoryManifestBlocks(FDirectoryManifest& Result, const FPath& Root, uin
 	}
 }
 
-static bool AlgorithmOptionsCompatible(const FAlgorithmOptions& A, const FAlgorithmOptions& B)
+static bool
+AlgorithmOptionsCompatible(const FAlgorithmOptions& A, const FAlgorithmOptions& B)
 {
-	return A.StrongHashAlgorithmId == B.StrongHashAlgorithmId
-		&& B.WeakHashAlgorithmId == B.WeakHashAlgorithmId
-		&& B.ChunkingAlgorithmId == B.ChunkingAlgorithmId;
+	return A.StrongHashAlgorithmId == B.StrongHashAlgorithmId && B.WeakHashAlgorithmId == B.WeakHashAlgorithmId &&
+		   B.ChunkingAlgorithmId == B.ChunkingAlgorithmId;
 }
 
 bool
@@ -2430,7 +2475,13 @@ SyncFile(const FNeedList&		   NeedList,
 		});
 
 		LogStatus(TargetFilePath.wstring().c_str(), L"Patching");
-		BuildTarget(*TargetFile, SourceFile, BaseDataReader, NeedList, Options.Algorithm.StrongHashAlgorithmId, Options.ProxyPool);
+		BuildTarget(*TargetFile,
+					SourceFile,
+					BaseDataReader,
+					NeedList,
+					Options.Algorithm.StrongHashAlgorithmId,
+					Options.ProxyPool,
+					Options.BlockCache);
 
 		Result.SourceBytes = ComputeSize(NeedList.Source);
 		Result.BaseBytes   = ComputeSize(NeedList.Base);
@@ -2631,22 +2682,21 @@ struct FPendingFileRename
 };
 
 // Updates the target directory manifest filename case to be consistent with reference.
-// Internally we always perform case-sensitive path comparisons, however on non-case-sensitive filesystems some local files may be renamed to a mismatching case.
-// We can update the locally-generated manifest to take the case from the reference manifest for equivalent paths.
+// Internally we always perform case-sensitive path comparisons, however on non-case-sensitive filesystems some local files may be renamed
+// to a mismatching case. We can update the locally-generated manifest to take the case from the reference manifest for equivalent paths.
 // Returns a list of files that should be renamed on disk.
 static std::vector<FPendingFileRename>
-FixManifestFileNameCases(
-	FDirectoryManifest& TargetDirectoryManifest,
-	const FDirectoryManifest& ReferenceManifest)
+FixManifestFileNameCases(FDirectoryManifest& TargetDirectoryManifest, const FDirectoryManifest& ReferenceManifest)
 {
 	// Build a lookup table of lowercase -> original file names and detect potential case conflicts (which will explode on Windows and Mac)
 
 	std::unordered_map<std::wstring, std::wstring> ReferenceFileNamesLowerCase;
-	bool bFoundCaseConflicts = false;
+	bool										   bFoundCaseConflicts = false;
 	for (auto& ReferenceManifestEntry : ReferenceManifest.Files)
 	{
 		std::wstring FileNameLowerCase = StringToLower(ReferenceManifestEntry.first);
-		auto InsertResult = ReferenceFileNamesLowerCase.insert(std::pair<std::wstring, std::wstring>(FileNameLowerCase, ReferenceManifestEntry.first));
+		auto		 InsertResult =
+			ReferenceFileNamesLowerCase.insert(std::pair<std::wstring, std::wstring>(FileNameLowerCase, ReferenceManifestEntry.first));
 
 		if (!InsertResult.second)
 		{
@@ -2670,7 +2720,7 @@ FixManifestFileNameCases(
 		if (ReferenceManifest.Files.find(TargetFileName) == ReferenceManifest.Files.end())
 		{
 			std::wstring TargetFileNameLowerCase = StringToLower(TargetFileName);
-			auto ReferenceIt = ReferenceFileNamesLowerCase.find(TargetFileNameLowerCase);
+			auto		 ReferenceIt			 = ReferenceFileNamesLowerCase.find(TargetFileNameLowerCase);
 			if (ReferenceIt != ReferenceFileNamesLowerCase.end())
 			{
 				FixupEntries.push_back({TargetFileName, ReferenceIt->second});
@@ -2700,7 +2750,7 @@ FixManifestFileNameCases(
 static bool
 FixFileNameCases(const FPath& RootPath, const std::vector<FPendingFileRename>& PendingRenames)
 {
-	std::vector<FPendingFileRename> UniqueRenames;
+	std::vector<FPendingFileRename>		   UniqueRenames;
 	std::unordered_set<FPath::string_type> UniqueRenamesSet;
 
 	// Build a rename schedule, with only unique entries (taking subdirectories into account)
@@ -2708,8 +2758,9 @@ FixFileNameCases(const FPath& RootPath, const std::vector<FPendingFileRename>& P
 	for (const FPendingFileRename& Entry : PendingRenames)
 	{
 		UNSYNC_ASSERTF(StringToLower(Entry.Old) == StringToLower(Entry.New),
-			L"FixFileNameCases expects inputs that are different only by case. Old: '%ls', New: '%ls'",
-			Entry.Old.c_str(), Entry.New.c_str());
+					   L"FixFileNameCases expects inputs that are different only by case. Old: '%ls', New: '%ls'",
+					   Entry.Old.c_str(),
+					   Entry.New.c_str());
 
 		FPath OldPath = Entry.Old;
 		FPath NewPath = Entry.New;
@@ -2739,8 +2790,7 @@ FixFileNameCases(const FPath& RootPath, const std::vector<FPendingFileRename>& P
 		}
 	}
 
-	std::sort(UniqueRenames.begin(), UniqueRenames.end(), [](const FPendingFileRename& A, const FPendingFileRename& B)
-	{
+	std::sort(UniqueRenames.begin(), UniqueRenames.end(), [](const FPendingFileRename& A, const FPendingFileRename& B) {
 		return A.Old < B.Old;
 	});
 
@@ -2811,7 +2861,7 @@ MergeManifests(FDirectoryManifest& Existing, const FDirectoryManifest& Other, bo
 		for (const auto& OtherFile : Other.Files)
 		{
 			std::wstring OtherNameLowerCase = StringToLower(OtherFile.first);
-			auto LowerCaseEntry = ExistingFileNamesLowerCase.find(OtherNameLowerCase);
+			auto		 LowerCaseEntry		= ExistingFileNamesLowerCase.find(OtherNameLowerCase);
 			if (LowerCaseEntry != ExistingFileNamesLowerCase.end())
 			{
 				// Remove file with conflicting case and add entry from the other manifest instead
@@ -2851,8 +2901,7 @@ DeleteUnnecessaryFiles(const FPath&				 TargetDirectory,
 	{
 		const std::wstring& TargetFileName = TargetManifestEntry.first;
 
-		auto Cleanup = [&ShouldCleanup, &TargetDirectory](const std::wstring& TargetFileName, const wchar_t* Reason) 
-		{
+		auto Cleanup = [&ShouldCleanup, &TargetDirectory](const std::wstring& TargetFileName, const wchar_t* Reason) {
 			FPath FilePath = TargetDirectory / TargetFileName;
 
 			if (!ShouldCleanup(TargetFileName))
@@ -2980,6 +3029,106 @@ GetAnonymizedHostName()
 	return Result;
 }
 
+struct FFileSyncTask
+{
+	const FFileManifest* SourceManifest = nullptr;
+	const FFileManifest* BaseManifest	= nullptr;
+	FPath				 OriginalSourceFilePath;
+	FPath				 ResolvedSourceFilePath;
+	FPath				 BaseFilePath;
+	FPath				 TargetFilePath;
+	FPath				 RelativeFilePath;
+	FNeedList			 NeedList;
+
+	uint64 NeedBytesFromSource = 0;
+	uint64 NeedBytesFromBase   = 0;
+	uint64 TotalSizeBytes	   = 0;
+
+	bool IsBaseValid() const { return !BaseFilePath.empty(); }
+};
+
+struct FFileSyncTaskBatch
+{
+	std::vector<const FFileSyncTask*> FileTasks;
+
+	uint64 TotalSizeBytes	   = 0;
+	uint64 NeedBytesFromSource = 0;
+
+	std::unique_ptr<FBlockCache> CreateBlockCache(FProxyPool& ProxyPool, EStrongHashAlgorithmID StrongHasher) const
+	{
+		std::unique_ptr<FBlockCache> Result = std::make_unique<FBlockCache>();
+
+		Result->BlockData.Resize(NeedBytesFromSource);
+
+		uint64 OutputCursor = 0;
+
+		HashSet<FHash128>		UniqueBlockSet;
+		std::vector<FNeedBlock> UniqueNeedBlocks;
+		for (const FFileSyncTask* Task : FileTasks)
+		{
+			for (const FNeedBlock& Block : Task->NeedList.Source)
+			{
+				if (UniqueBlockSet.insert(Block.Hash.ToHash128()).second)
+				{
+					UniqueNeedBlocks.push_back(Block);
+				}
+			}
+		}
+
+		Result->BlockMap.reserve(UniqueNeedBlocks.size());
+
+		ProxyPool.ParallelDownloadSemaphore.Acquire();
+		std::unique_ptr<FProxy> Proxy = ProxyPool.Alloc();
+
+		if (Proxy)
+		{
+			auto DownloadCallback = [&OutputCursor, &Result, &UniqueBlockSet, StrongHasher](const FDownloadedBlock& Block,
+																							FHash128				BlockHash) {
+				if (OutputCursor + Block.DecompressedSize <= Result->BlockData.Size())
+				{
+					if (UniqueBlockSet.find(BlockHash) != UniqueBlockSet.end())
+					{
+						FMutBufferView OutputView = Result->BlockData.MutView(OutputCursor, Block.DecompressedSize);
+
+						bool bOk = true;
+
+						if (Block.IsCompressed())
+						{
+							bOk = Decompress(Block.Data, Block.CompressedSize, OutputView.Data, OutputView.Size);
+						}
+						else
+						{
+							memcpy(OutputView.Data, Block.Data, OutputView.Size);
+						}
+
+						if (bOk)
+						{
+							const FHash128 ActualBlockHash = ComputeHash(OutputView.Data, OutputView.Size, StrongHasher).ToHash128();
+
+							bOk = BlockHash == ActualBlockHash;
+						}
+
+						if (bOk)
+						{
+							Result->BlockMap[BlockHash] = FBufferView{OutputView.Data, OutputView.Size};
+							OutputCursor += Block.DecompressedSize;
+						}
+					}
+				}
+			};
+
+			FDownloadResult DownloadResult =
+				Proxy->Download(MakeView<FNeedBlock>(UniqueNeedBlocks.data(), UniqueNeedBlocks.size()), DownloadCallback);
+
+			UNSYNC_UNUSED(DownloadResult);
+		}
+
+		ProxyPool.ParallelDownloadSemaphore.Release();
+
+		return Result;
+	}
+};
+
 bool  // TODO: return a TResult
 SyncDirectory(const FSyncDirectoryOptions& SyncOptions)
 {
@@ -3058,11 +3207,11 @@ SyncDirectory(const FSyncDirectoryOptions& SyncOptions)
 		for (const FPath& ThisSourcePath : AllSources)
 		{
 			if (!LoadAndMergeSourceManifest(SourceDirectoryManifest,
-				ThisSourcePath,
-				TargetTempPath,
-				SyncFilter,
-				SourceManifestOverride,
-				bCaseSensitiveTargetFileSystem))
+											ThisSourcePath,
+											TargetTempPath,
+											SyncFilter,
+											SourceManifestOverride,
+											bCaseSensitiveTargetFileSystem))
 			{
 				return false;
 			}
@@ -3137,24 +3286,7 @@ SyncDirectory(const FSyncDirectoryOptions& SyncOptions)
 	std::atomic<uint64> StatSourceBytes = {};
 	std::atomic<uint64> StatBaseBytes	= {};
 
-	struct FileSyncTask
-	{
-		const FFileManifest* SourceManifest = nullptr;
-		const FFileManifest* BaseManifest	= nullptr;
-		FPath				 OriginalSourceFilePath;
-		FPath				 ResolvedSourceFilePath;
-		FPath				 BaseFilePath;
-		FPath				 TargetFilePath;
-		FPath				 RelativeFilePath;
-		FNeedList			 NeedList;
-
-		uint64 NeedBytesFromSource = 0;
-		uint64 NeedBytesFromBase   = 0;
-		uint64 TotalSizeBytes	   = 0;
-
-		bool IsBaseValid() const { return !BaseFilePath.empty(); }
-	};
-	std::deque<FileSyncTask> SyncTaskList;
+	std::vector<FFileSyncTask> AllFileTasks;
 
 	LogGlobalStatus(L"Scanning base directory");
 	UNSYNC_VERBOSE(L"Scanning base directory");
@@ -3176,8 +3308,8 @@ SyncDirectory(const FSyncDirectoryOptions& SyncOptions)
 	bool			   bBaseDirectoryManifestValid = false;
 	bool			   bQuickDifferencePossible	   = false;
 
-	if (!SyncOptions.bFullDifference && SourceDirectoryManifest.Options.ChunkingAlgorithmId == EChunkingAlgorithmID::VariableBlocks
-		&& PathExists(BaseManifestPath))
+	if (!SyncOptions.bFullDifference && SourceDirectoryManifest.Options.ChunkingAlgorithmId == EChunkingAlgorithmID::VariableBlocks &&
+		PathExists(BaseManifestPath))
 	{
 		bBaseDirectoryManifestValid = LoadDirectoryManifest(BaseDirectoryManifest, BasePath, BaseManifestPath);
 		if (bBaseDirectoryManifestValid && AlgorithmOptionsCompatible(SourceDirectoryManifest.Options, TargetDirectoryManifest.Options))
@@ -3209,7 +3341,7 @@ SyncDirectory(const FSyncDirectoryOptions& SyncOptions)
 		TotalSourceSize += SourceFileManifest.Size;
 
 		bool bTargetFileAttributesMatch = false;
-		auto TargetManifestIt = TargetDirectoryManifest.Files.find(SourceFilename);
+		auto TargetManifestIt			= TargetDirectoryManifest.Files.find(SourceFilename);
 		if (TargetManifestIt != TargetDirectoryManifest.Files.end())
 		{
 			const FFileManifest& TargetFileManifest = TargetManifestIt->second;
@@ -3284,7 +3416,7 @@ SyncDirectory(const FSyncDirectoryOptions& SyncOptions)
 				{
 					UNSYNC_VERBOSE2(L"Dirty file: '%ls'", SourceManifestIt.first.c_str());
 				}
-				
+
 				StatPartialCopy++;
 
 				if (bFileSystemSource && SyncOptions.bValidateSourceFiles && !SourceAttribCache.Exists(ResolvedSourceFilePath) &&
@@ -3295,7 +3427,7 @@ SyncDirectory(const FSyncDirectoryOptions& SyncOptions)
 				}
 			}
 
-			FileSyncTask Task;
+			FFileSyncTask Task;
 			Task.OriginalSourceFilePath = std::move(SourceFilePath);
 			Task.ResolvedSourceFilePath = std::move(ResolvedSourceFilePath);
 			Task.BaseFilePath			= std::move(BaseFilePath);
@@ -3319,7 +3451,7 @@ SyncDirectory(const FSyncDirectoryOptions& SyncOptions)
 				}
 			}
 
-			SyncTaskList.push_back(Task);
+			AllFileTasks.push_back(Task);
 		}
 	}
 
@@ -3339,12 +3471,13 @@ SyncDirectory(const FSyncDirectoryOptions& SyncOptions)
 
 	uint64 EstimatedNeedBytesFromSource = 0;
 	uint64 EstimatedNeedBytesFromBase	= 0;
+	uint64 TotalSyncSizeBytes			= 0;
 
 	{
 		UNSYNC_LOG_INDENT;
 		auto TimeDiffBegin = TimePointNow();
 
-		auto DiffTask = [Algorithm, bQuickDifferencePossible](FileSyncTask& Item) {
+		auto DiffTask = [Algorithm, bQuickDifferencePossible](FFileSyncTask& Item) {
 			FLogVerbosityScope VerbosityScope(false);  // turn off logging from threads
 
 			const FGenericBlockArray& SourceBlocks = Item.SourceManifest->Blocks;
@@ -3402,32 +3535,18 @@ SyncDirectory(const FSyncDirectoryOptions& SyncOptions)
 			UNSYNC_ASSERT(Item.NeedBytesFromSource + Item.NeedBytesFromBase == Item.TotalSizeBytes);
 		};
 
-		ParallelForEach(SyncTaskList.begin(), SyncTaskList.end(), DiffTask);
+		ParallelForEach(AllFileTasks.begin(), AllFileTasks.end(), DiffTask);
 
 		auto TimeDiffEnd = TimePointNow();
 
 		double Duration = DurationSec(TimeDiffBegin, TimeDiffEnd);
 		UNSYNC_VERBOSE(L"Difference complete in %.3f sec", Duration);
 
-		HashSet<FHash128> UniqueSourceBlocks;
-		uint64			  UniqueSourceBytes = 0;
-
-		uint64 TotalSyncSizeBytes = 0;
-
-		for (FileSyncTask& Item : SyncTaskList)
+		for (FFileSyncTask& Item : AllFileTasks)
 		{
 			EstimatedNeedBytesFromSource += Item.NeedBytesFromSource;
 			EstimatedNeedBytesFromBase += Item.NeedBytesFromBase;
 			TotalSyncSizeBytes += Item.TotalSizeBytes;
-
-			for (const auto& Block : Item.NeedList.Source)
-			{
-				auto InsertRes = UniqueSourceBlocks.insert(Block.Hash.ToHash128());	 // #wip-widehash
-				if (InsertRes.second)
-				{
-					UniqueSourceBytes += Block.Size;
-				}
-			}
 		}
 
 		UNSYNC_VERBOSE(L"Total need from source: %.2f MB", SizeMb(EstimatedNeedBytesFromSource));
@@ -3439,8 +3558,10 @@ SyncDirectory(const FSyncDirectoryOptions& SyncOptions)
 			UNSYNC_ERROR(
 				L"Sync requires %.0f MB (%llu bytes) of disk space, but only %.0f MB (%llu bytes) is available. "
 				L"Use --no-space-validation flag to suppress this check.",
-				SizeMb(TotalSyncSizeBytes), TotalSyncSizeBytes,
-				SizeMb(AvailableDiskBytes), AvailableDiskBytes);
+				SizeMb(TotalSyncSizeBytes),
+				TotalSyncSizeBytes,
+				SizeMb(AvailableDiskBytes),
+				AvailableDiskBytes);
 			return false;
 		}
 	}
@@ -3471,7 +3592,7 @@ SyncDirectory(const FSyncDirectoryOptions& SyncOptions)
 			UNSYNC_VERBOSE(L"Connection established");
 			UNSYNC_VERBOSE(L"Building block request map");
 
-			bool bProxyHasData = Proxy->Contains(SourceDirectoryManifest);
+			const bool bProxyHasData = Proxy->Contains(SourceDirectoryManifest);
 
 			ProxyPool.Dealloc(std::move(Proxy));
 
@@ -3479,7 +3600,7 @@ SyncDirectory(const FSyncDirectoryOptions& SyncOptions)
 			{
 				ProxyPool.InitRequestMap(SourceDirectoryManifest.Options.StrongHashAlgorithmId);
 
-				for (const FileSyncTask& Task : SyncTaskList)
+				for (const FFileSyncTask& Task : AllFileTasks)
 				{
 					ProxyPool.BuildFileBlockRequests(Task.OriginalSourceFilePath, Task.ResolvedSourceFilePath, *Task.SourceManifest);
 				}
@@ -3508,30 +3629,70 @@ SyncDirectory(const FSyncDirectoryOptions& SyncOptions)
 	UNSYNC_VERBOSE(L"Copying files");
 
 	{
-		struct BackgroundTaskResult
+		struct FBackgroundTaskResult
 		{
 			FPath			TargetFilePath;
 			FFileSyncResult SyncResult;
 			bool			bIsPartialCopy = false;
 		};
 
-		std::mutex						  BackgroundTaskStatMutex;
-		std::vector<BackgroundTaskResult> BackgroundTaskResults;
+		std::deque<FFileSyncTaskBatch> SyncTaskList;
+
+		std::mutex						   BackgroundTaskStatMutex;
+		std::vector<FBackgroundTaskResult> BackgroundTaskResults;
 
 		// Tasks are sorted by download size and processed by multiple threads.
 		// Large downloads are processed on the foreground thread and small ones on the background.
-		std::sort(SyncTaskList.begin(), SyncTaskList.end(), [](const FileSyncTask& A, const FileSyncTask& B) {
+		std::sort(AllFileTasks.begin(), AllFileTasks.end(), [](const FFileSyncTask& A, const FFileSyncTask& B) {
 			return A.NeedBytesFromSource < B.NeedBytesFromSource;
 		});
 
-		auto SyncTaskBody = [Algorithm,
-							 &SyncOptions,
-							 &StatSourceBytes,
-							 &StatBaseBytes,
-							 &BackgroundTaskStatMutex,
-							 &BackgroundTaskResults,
-							 &NumFailedTasks,
-							 &ProxyPool](const FileSyncTask& Item, bool bBackground) {
+		// Blocks for multiple files can be downloaded in one request.
+		// Group small file tasks into batches to reduce the number of individual download requests.
+		{
+			const uint64 MaxBatchDownloadSize = 4_MB;
+
+			FFileSyncTaskBatch CurrentBatch;
+			for (const FFileSyncTask& FileTask : AllFileTasks)
+			{
+				if (CurrentBatch.FileTasks.size() && CurrentBatch.NeedBytesFromSource + FileTask.NeedBytesFromSource > MaxBatchDownloadSize)
+				{
+					SyncTaskList.push_back(std::move(CurrentBatch));
+					CurrentBatch = {};
+				}
+
+				CurrentBatch.FileTasks.push_back(&FileTask);
+				CurrentBatch.NeedBytesFromSource += FileTask.NeedBytesFromSource;
+				CurrentBatch.TotalSizeBytes += FileTask.TotalSizeBytes;
+			}
+
+			if (CurrentBatch.FileTasks.size())
+			{
+				SyncTaskList.push_back(std::move(CurrentBatch));
+			}
+		}
+
+		// Validate batching
+		{
+			uint64 TotalSyncSizeBatched = 0;
+			uint64 TotalFilesBatched	= 0;
+			for (const FFileSyncTaskBatch& Batch : SyncTaskList)
+			{
+				TotalSyncSizeBatched += Batch.TotalSizeBytes;
+				TotalFilesBatched += Batch.FileTasks.size();
+			}
+			UNSYNC_ASSERT(TotalFilesBatched == AllFileTasks.size());
+			UNSYNC_ASSERT(TotalSyncSizeBatched == TotalSyncSizeBytes);
+		}
+
+		auto FileSyncTaskBody = [Algorithm,
+								 &SyncOptions,
+								 &StatSourceBytes,
+								 &StatBaseBytes,
+								 &BackgroundTaskStatMutex,
+								 &BackgroundTaskResults,
+								 &NumFailedTasks,
+								 &ProxyPool](const FFileSyncTask& Item, FBlockCache* BlockCache, bool bBackground) {
 			UNSYNC_VERBOSE(L"Copy '%ls' (%ls)", Item.TargetFilePath.wstring().c_str(), (Item.NeedBytesFromBase) ? L"partial" : L"full");
 
 			std::unique_ptr<NativeFile> BaseFile;
@@ -3547,6 +3708,7 @@ SyncDirectory(const FSyncDirectoryOptions& SyncOptions)
 			SyncFileOptions.Algorithm			 = Algorithm;
 			SyncFileOptions.BlockSize			 = SourceBlockSize;
 			SyncFileOptions.ProxyPool			 = &ProxyPool;
+			SyncFileOptions.BlockCache			 = BlockCache;
 			SyncFileOptions.bValidateTargetFiles = SyncOptions.bValidateTargetFiles;
 
 			FFileSyncResult SyncResult =
@@ -3566,7 +3728,7 @@ SyncDirectory(const FSyncDirectoryOptions& SyncOptions)
 
 			if (bBackground)
 			{
-				BackgroundTaskResult Result;
+				FBackgroundTaskResult Result;
 				Result.TargetFilePath = Item.TargetFilePath;
 				Result.SyncResult	  = SyncResult;
 				Result.bIsPartialCopy = Item.NeedBytesFromBase != 0;
@@ -3598,26 +3760,43 @@ SyncDirectory(const FSyncDirectoryOptions& SyncOptions)
 		std::atomic<uint64> BackgroundTaskMemory	   = {};
 		std::atomic<uint64> RemainingSourceBytes	   = EstimatedNeedBytesFromSource;
 
-		std::mutex SchedulerMutex;
+		std::mutex				SchedulerMutex;
 		std::condition_variable SchedulerEvent;
 
 		while (!SyncTaskList.empty())
 		{
 			if (NumForegroundTasks == 0)
 			{
-				FileSyncTask LocalTask = std::move(SyncTaskList.back());
+				FFileSyncTaskBatch LocalTaskBatch = std::move(SyncTaskList.back());
+
 				SyncTaskList.pop_back();
 				++NumForegroundTasks;
 
-				RemainingSourceBytes -= LocalTask.NeedBytesFromSource;
+				RemainingSourceBytes -= LocalTaskBatch.NeedBytesFromSource;
 
-				ForegroundTaskGroup.run(
-					[Task = std::move(LocalTask), &SchedulerEvent, &NumForegroundTasks, &SyncTaskBody, LogVerbose = GLogVerbose]() {
-						FLogVerbosityScope VerbosityScope(LogVerbose);
-						SyncTaskBody(Task, false);
-						--NumForegroundTasks;
-						SchedulerEvent.notify_one();
-					});
+				ForegroundTaskGroup.run([TaskBatch = std::move(LocalTaskBatch),
+										 &SchedulerEvent,
+										 &NumForegroundTasks,
+										 &FileSyncTaskBody,
+										 &ProxyPool,
+										 Algorithm,
+										 LogVerbose = GLogVerbose]() {
+					FLogVerbosityScope VerbosityScope(LogVerbose);
+
+					std::unique_ptr<FBlockCache> BlockCache;
+					if (TaskBatch.FileTasks.size() > 1 && ProxyPool.IsValid())
+					{
+						BlockCache = TaskBatch.CreateBlockCache(ProxyPool, Algorithm.StrongHashAlgorithmId);
+					}
+
+					for (const FFileSyncTask* Task : TaskBatch.FileTasks)
+					{
+						FileSyncTaskBody(*Task, BlockCache.get(), false);
+					}
+
+					--NumForegroundTasks;
+					SchedulerEvent.notify_one();
+				});
 				continue;
 			}
 
@@ -3626,22 +3805,37 @@ SyncDirectory(const FSyncDirectoryOptions& SyncOptions)
 			if (NumBackgroundTasks < MaxBackgroundTasks && (SyncTaskList.front().NeedBytesFromSource < RemainingSourceBytes / 4) &&
 				(BackgroundTaskMemory + SyncTaskList.front().TotalSizeBytes < BackgroundTaskMemoryBudget))
 			{
-				FileSyncTask LocalTask = std::move(SyncTaskList.front());
+				FFileSyncTaskBatch LocalTaskBatch = std::move(SyncTaskList.front());
 				SyncTaskList.pop_front();
 
-				BackgroundTaskMemory += LocalTask.TotalSizeBytes;
+				BackgroundTaskMemory += LocalTaskBatch.TotalSizeBytes;
 				++NumBackgroundTasks;
 
-				RemainingSourceBytes -= LocalTask.NeedBytesFromSource;
+				RemainingSourceBytes -= LocalTaskBatch.NeedBytesFromSource;
 
-				BackgroundTaskGroup.run(
-					[Task = std::move(LocalTask), &SchedulerEvent, &NumBackgroundTasks, &SyncTaskBody, &BackgroundTaskMemory]() {
-						FLogVerbosityScope VerbosityScope(false);  // turn off logging from background threads
-						SyncTaskBody(Task, true);
-						--NumBackgroundTasks;
-						BackgroundTaskMemory -= Task.TotalSizeBytes;
-						SchedulerEvent.notify_one();
-					});
+				BackgroundTaskGroup.run([TaskBatch = std::move(LocalTaskBatch),
+										 &SchedulerEvent,
+										 &NumBackgroundTasks,
+										 &FileSyncTaskBody,
+										 &BackgroundTaskMemory,
+										 &ProxyPool,
+										 Algorithm]() {
+					FLogVerbosityScope VerbosityScope(false);  // turn off logging from background threads
+
+					std::unique_ptr<FBlockCache> BlockCache;
+					if (TaskBatch.FileTasks.size() > 1 && ProxyPool.IsValid())
+					{
+						BlockCache = TaskBatch.CreateBlockCache(ProxyPool, Algorithm.StrongHashAlgorithmId);
+					}
+
+					for (const FFileSyncTask* Task : TaskBatch.FileTasks)
+					{
+						FileSyncTaskBody(*Task, BlockCache.get(), true);
+					}
+					BackgroundTaskMemory -= TaskBatch.TotalSizeBytes;
+					--NumBackgroundTasks;
+					SchedulerEvent.notify_one();
+				});
 
 				continue;
 			}
@@ -3660,16 +3854,16 @@ SyncDirectory(const FSyncDirectoryOptions& SyncOptions)
 
 		UNSYNC_ASSERT(RemainingSourceBytes == 0);
 
-		bool bAllBackgroundTasksSucceeded = true;
-		uint32 NumBackgroundSyncFiles = 0;
-		uint64 DownloadedBackgroundBytes = 0;
-		for (const BackgroundTaskResult& Item : BackgroundTaskResults)
+		bool   bAllBackgroundTasksSucceeded = true;
+		uint32 NumBackgroundSyncFiles		= 0;
+		uint64 DownloadedBackgroundBytes	= 0;
+		for (const FBackgroundTaskResult& Item : BackgroundTaskResults)
 		{
 			if (Item.SyncResult.Succeeded())
 			{
 				UNSYNC_VERBOSE2(L"Copied '%ls' (%ls, background)",
-							   Item.TargetFilePath.wstring().c_str(),
-							   Item.bIsPartialCopy ? L"partial" : L"full");
+								Item.TargetFilePath.wstring().c_str(),
+								Item.bIsPartialCopy ? L"partial" : L"full");
 				++NumBackgroundSyncFiles;
 				DownloadedBackgroundBytes += Item.SyncResult.SourceBytes;
 			}
@@ -3686,7 +3880,7 @@ SyncDirectory(const FSyncDirectoryOptions& SyncOptions)
 
 		if (!bAllBackgroundTasksSucceeded)
 		{
-			for (const BackgroundTaskResult& Item : BackgroundTaskResults)
+			for (const FBackgroundTaskResult& Item : BackgroundTaskResults)
 			{
 				if (!Item.SyncResult.Succeeded())
 				{
@@ -4159,7 +4353,7 @@ FSyncFilter::ShouldSync(const FPath& Filename) const
 bool
 FSyncFilter::ShouldSync(const std::wstring& Filename) const
 {
-	bool bInclude = SyncIncludedWords.empty(); // Include everything if there are no specific inclusions
+	bool bInclude = SyncIncludedWords.empty();	// Include everything if there are no specific inclusions
 	for (const std::wstring& Word : SyncIncludedWords)
 	{
 		if (Filename.find(Word) != std::wstring::npos)
@@ -4173,7 +4367,7 @@ FSyncFilter::ShouldSync(const std::wstring& Filename) const
 	{
 		return false;
 	}
-	
+
 	for (const std::wstring& Word : SyncExcludedWords)
 	{
 		if (Filename.find(Word) != std::wstring::npos)
diff --git a/Engine/Source/Programs/Unsync/Private/UnsyncCore.h b/Engine/Source/Programs/Unsync/Private/UnsyncCore.h
--- a/Engine/Source/Programs/Unsync/Private/UnsyncCore.h
+++ b/Engine/Source/Programs/Unsync/Private/UnsyncCore.h
@@ -20,6 +20,7 @@ extern bool GDryRun;
 
 class FProxy;
 class FProxyPool;
+class FBlockCache;
 struct FRemoteDesc;
 struct FIOReader;
 struct FIOWriter;
@@ -248,13 +249,15 @@ void BuildTarget(FIOWriter&				Result,
 				 FIOReader&				Base,
 				 const FNeedList&		NeedList,
 				 EStrongHashAlgorithmID StrongHasher,
-				 FProxyPool*			ProxyPool = nullptr);
+				 FProxyPool*			ProxyPool = nullptr,
+				 FBlockCache*			BlockCache = nullptr);
 
 FBuffer BuildTargetBuffer(FIOReader&			 SourceProvider,
 						  FIOReader&			 BaseProvider,
 						  const FNeedList&		 NeedList,
 						  EStrongHashAlgorithmID StrongHasher,
-						  FProxyPool*			 ProxyPool = nullptr);
+						  FProxyPool*			 ProxyPool = nullptr,
+						  FBlockCache*			 BlockCache = nullptr);
 
 FBuffer BuildTargetBuffer(const uint8*			 SourceData,
 						  uint64				 SourceSize,
@@ -306,6 +309,7 @@ struct FSyncFileOptions
 	uint32			  BlockSize = uint32(64_KB);
 
 	FProxyPool* ProxyPool = nullptr;
+	FBlockCache* BlockCache = nullptr;
 
 	bool bValidateTargetFiles = true;  // WARNING: turning this off is intended only for testing/profiling
 };
