unsync - Fix resolving mapped network drive paths

Explicitly specify UNIVERSAL_NAME_INFOW instead of UNIVERSAL_NAME_INFO that depends on the compiler flags

diff --git a/Engine/Source/Programs/Unsync/Private/UnsyncUtil.cpp b/Engine/Source/Programs/Unsync/Private/UnsyncUtil.cpp
--- a/Engine/Source/Programs/Unsync/Private/UnsyncUtil.cpp
+++ b/Engine/Source/Programs/Unsync/Private/UnsyncUtil.cpp
@@ -276,9 +276,11 @@ FPath GetUniversalPath(const FPath& Path)
 	// https://docs.microsoft.com/en-us/windows/win32/api/winnetwk/nf-winnetwk-wnetgetuniversalnamea
 
 	static constexpr DWORD MaxBufferSize = 1024;
-	WCHAR Buffer[MaxBufferSize];
-	DWORD BufferSize = MaxBufferSize;
-	UNIVERSAL_NAME_INFO* UniversalNameInfo = (UNIVERSAL_NAME_INFO*)Buffer;
+
+	WCHAR Buffer[MaxBufferSize] = {};
+
+	DWORD				 BufferSize		   = MaxBufferSize;
+	UNIVERSAL_NAME_INFOW* UniversalNameInfo = (UNIVERSAL_NAME_INFOW*)Buffer;
 
 	DWORD ErrorCode = WNetGetUniversalNameW(Path.native().c_str(), UNIVERSAL_NAME_INFO_LEVEL, (LPVOID)UniversalNameInfo, &BufferSize);
 	if (ErrorCode == NO_ERROR)
