unsync - Fix global progress reporting

diff --git a/Engine/Source/Programs/Unsync/Private/UnsyncCore.cpp b/Engine/Source/Programs/Unsync/Private/UnsyncCore.cpp
--- a/Engine/Source/Programs/Unsync/Private/UnsyncCore.cpp
+++ b/Engine/Source/Programs/Unsync/Private/UnsyncCore.cpp
@@ -1532,8 +1532,8 @@ BuildTarget(FIOWriter&			   Result,
 			if (It != BlockCache->BlockMap.end())
 			{
 				FBufferView BlockBuffer = It->second;
-				Result.Write(BlockBuffer.Data, NeedBlock.TargetOffset, BlockBuffer.Size);
 
+				Result.Write(BlockBuffer.Data, NeedBlock.TargetOffset, BlockBuffer.Size);
 				AddGlobalProgress(NeedBlock.Size, EListType::Source);
 			}
 		}
@@ -1680,6 +1680,8 @@ BuildTarget(FIOWriter&			   Result,
 								if (bOk)
 								{
 									uint64 WrittenBytes = Result.Write(DecompressedData.GetData(), Cmd.TargetOffset, Cmd.Size);
+									AddGlobalProgress(Cmd.Size, EListType::Source);
+
 									if (WrittenBytes != Cmd.Size)
 									{
 										bOk = false;
