unsync - Initial implementation of the "scavenge" mode which aims to reduce download times by reusing data outside of the final sync directory

This is useful for people who want to download each data set into a uniquely named directory, instead of always patching the last downloaded version.

Current implementation is highly experimental and should only be used for testing purposes.

diff --git a/Engine/Source/Programs/Unsync/CMakeLists.txt b/Engine/Source/Programs/Unsync/CMakeLists.txt
--- a/Engine/Source/Programs/Unsync/CMakeLists.txt
+++ b/Engine/Source/Programs/Unsync/CMakeLists.txt
@@ -55,6 +55,7 @@ set(src
 	Private/UnsyncFile.h
 	Private/UnsyncHash.cpp
 	Private/UnsyncHash.h
+	Private/UnsyncHashTable.h
 	Private/UnsyncHttp.cpp
 	Private/UnsyncHttp.h
 	Private/UnsyncJupiter.cpp
@@ -74,6 +75,8 @@ set(src
 	Private/UnsyncRemote.cpp
 	Private/UnsyncRemote.h
 	Private/UnsyncScan.h
+	Private/UnsyncScavenger.cpp
+	Private/UnsyncScavenger.h
 	Private/UnsyncSerialization.cpp
 	Private/UnsyncSerialization.h
 	Private/UnsyncSocket.cpp
diff --git a/Engine/Source/Programs/Unsync/Private/UnsyncCmdSync.cpp b/Engine/Source/Programs/Unsync/Private/UnsyncCmdSync.cpp
--- a/Engine/Source/Programs/Unsync/Private/UnsyncCmdSync.cpp
+++ b/Engine/Source/Programs/Unsync/Private/UnsyncCmdSync.cpp
@@ -130,6 +130,8 @@ CmdSync(const FCmdSyncOptions& Options)
 			SyncOptions.Source				   = ResolvedSource;
 			SyncOptions.Base				   = Options.Target;  // read base data from existing target
 			SyncOptions.Target				   = Options.Target;
+			SyncOptions.ScavengeRoot		   = Options.ScavengeRoot;
+			SyncOptions.ScavengeDepth		   = Options.ScavengeDepth;
 			SyncOptions.Overlays			   = ResolvedOverlays;
 			SyncOptions.SourceManifestOverride = Options.SourceManifestOverride;
 			SyncOptions.Remote				   = &Options.Remote;
diff --git a/Engine/Source/Programs/Unsync/Private/UnsyncCmdSync.h b/Engine/Source/Programs/Unsync/Private/UnsyncCmdSync.h
--- a/Engine/Source/Programs/Unsync/Private/UnsyncCmdSync.h
+++ b/Engine/Source/Programs/Unsync/Private/UnsyncCmdSync.h
@@ -15,6 +15,8 @@ struct FCmdSyncOptions
 	FPath Source;
 	FPath Target;
 	FPath SourceManifestOverride;
+	FPath ScavengeRoot;
+	uint32 ScavengeDepth = 5;
 
 	std::vector<FPath> Overlays;
 
diff --git a/Engine/Source/Programs/Unsync/Private/UnsyncCore.cpp b/Engine/Source/Programs/Unsync/Private/UnsyncCore.cpp
--- a/Engine/Source/Programs/Unsync/Private/UnsyncCore.cpp
+++ b/Engine/Source/Programs/Unsync/Private/UnsyncCore.cpp
@@ -5,8 +5,10 @@
 #include "UnsyncCompression.h"
 #include "UnsyncCore.h"
 #include "UnsyncFile.h"
+#include "UnsyncHashTable.h"
 #include "UnsyncProxy.h"
 #include "UnsyncScan.h"
+#include "UnsyncScavenger.h"
 #include "UnsyncSerialization.h"
 #include "UnsyncThread.h"
 #include "UnsyncUtil.h"
@@ -22,23 +24,14 @@
 UNSYNC_THIRD_PARTY_INCLUDES_START
 #include <blake3.h>
 #include <md5-sse2.h>
-#include <flat_hash_map.hpp>
 UNSYNC_THIRD_PARTY_INCLUDES_END
 
-#define UNSYNC_VERSION_STR "1.0.51-dev"
+#define UNSYNC_VERSION_STR "1.0.51-dev3"
 
 namespace unsync {
 
 bool GDryRun = false;
 
-template<typename K, typename V, typename Hash = std::hash<K>, typename Eq = std::equal_to<K>>
-// using HashMap = std::unordered_map<K,V,Hash,Eq>;
-using HashMap = ska::flat_hash_map<K, V, Hash, Eq>;
-
-template<typename K, typename Hash = std::hash<K>, typename Eq = std::equal_to<K>>
-// using HashSet = std::unordered_set<K, Hash, Eq>;
-using HashSet = ska::flat_hash_set<K, Hash, Eq>;
-
 static const uint32 MAX_ACTIVE_READERS = 64;  // std::max(4, std::thread::hardware_concurrency();
 
 inline uint32
@@ -79,9 +72,9 @@ struct FLogProgressScope
 
 	FLogProgressScope(uint64 InTotal, ELogProgressUnits InUnits = ELogProgressUnits::Raw, uint64 InPeriodMilliseconds = 500)
 	: Current(0)
-	, TOTAL(InTotal)
-	, PERIOD_MILLISECONDS(InPeriodMilliseconds)
-	, UNITS(InUnits)
+	, Total(InTotal)
+	, PeriodMilliseconds(InPeriodMilliseconds)
+	, Units(InUnits)
 	, NextProgressLogTime(TimePointNow())
 	, bEnabled(true)
 	{
@@ -100,37 +93,37 @@ struct FLogProgressScope
 		Current += X;
 
 		std::lock_guard<std::mutex> LockGuard(Mutex);
-		const uint64				CurrentClamped = std::min<uint64>(Current, TOTAL);
-		const bool					bComplete	   = (CurrentClamped == TOTAL) || bForceComplete;
+		const uint64				CurrentClamped = std::min<uint64>(Current, Total);
+		const bool					bComplete	   = (CurrentClamped == Total) || bForceComplete;
 		if (bEnabled && GLogVerbose && (TimePointNow() > NextProgressLogTime || bComplete))
 		{
 			const wchar_t* Ending = bComplete ? L"\n" : L"\r";
-			switch (UNITS)
+			switch (Units)
 			{
 				case ELogProgressUnits::GB:
 					LogPrintf(ELogLevel::Debug,
 							  L"%.2f / %.2f GB%ls",
 							  double(CurrentClamped) / double(1_GB),
-							  double(TOTAL) / double(1_GB),
+							  double(Total) / double(1_GB),
 							  Ending);
 					break;
 				case ELogProgressUnits::MB:
 					LogPrintf(ELogLevel::Debug,
 							  L"%.2f / %.2f MB%ls",
 							  double(CurrentClamped) / double(1_MB),
-							  double(TOTAL) / double(1_MB),
+							  double(Total) / double(1_MB),
 							  Ending);
 					break;
 				case ELogProgressUnits::Bytes:
-					LogPrintf(ELogLevel::Debug, L"%llu / %llu bytes%ls", (uint64)CurrentClamped, TOTAL, Ending);
+					LogPrintf(ELogLevel::Debug, L"%llu / %llu bytes%ls", (uint64)CurrentClamped, Total, Ending);
 					break;
 				case ELogProgressUnits::Raw:
 				default:
-					LogPrintf(ELogLevel::Debug, L"%llu / %llu%ls", (uint64)CurrentClamped, TOTAL, Ending);
+					LogPrintf(ELogLevel::Debug, L"%llu / %llu%ls", (uint64)CurrentClamped, Total, Ending);
 					break;
 			}
 
-			NextProgressLogTime = TimePointNow() + std::chrono::milliseconds(PERIOD_MILLISECONDS);
+			NextProgressLogTime = TimePointNow() + std::chrono::milliseconds(PeriodMilliseconds);
 			LogGlobalProgress();
 
 			if (bComplete)
@@ -142,9 +135,9 @@ struct FLogProgressScope
 
 	std::mutex				Mutex;
 	std::atomic<uint64>		Current;
-	const uint64			TOTAL;
-	const uint64			PERIOD_MILLISECONDS;
-	const ELogProgressUnits UNITS;
+	const uint64			Total;
+	const uint64			PeriodMilliseconds;
+	const ELogProgressUnits Units;
 	FTimePoint				NextProgressLogTime;
 	std::atomic<bool>		bEnabled;
 };
@@ -417,8 +410,8 @@ ComputeBlocksVariableT(FIOReader& Reader, uint32 BlockSize, EStrongHashAlgorithm
 
 	uint64 NumTotalBlocks = 0;
 
-	HashSet<FGenericHash> UniqueBlockSet;
-	FGenericBlockArray	  UniqueBlocks;
+	THashSet<FGenericHash> UniqueBlockSet;
+	FGenericBlockArray	   UniqueBlocks;
 	for (const FGenericBlock& It : Result)
 	{
 		auto InsertResult = UniqueBlockSet.insert(It.HashStrong);
@@ -611,7 +604,7 @@ ComputeBlocksFixedT(FIOReader& Reader, uint32 BlockSize, EStrongHashAlgorithmID
 	double Duration = DurationSec(TimeBegin, TimePointNow());
 	UNSYNC_VERBOSE2(L"Done in %.3f sec (%.3f MB / sec)", Duration, SizeMb((double(Reader.GetSize()) / Duration)));
 
-	HashSet<uint32> UniqueHashes;
+	THashSet<uint32> UniqueHashes;
 	for (const auto& It : Blocks)
 	{
 		UniqueHashes.insert(It.HashWeak);
@@ -750,8 +743,8 @@ DiffManifestBlocksT(const std::vector<BlockType>& SourceBlocks, const std::vecto
 		uint64 Count = 0;
 	};
 
-	HashMap<typename BlockType::StrongHashType, BlockIndexAndCount> BaseBlockMap;
-	HashMap<uint64, uint64>											BaseBlockByOffset;
+	THashMap<typename BlockType::StrongHashType, BlockIndexAndCount> BaseBlockMap;
+	THashMap<uint64, uint64>										 BaseBlockByOffset;
 
 	for (uint64 I = 0; I < BaseBlocks.size(); ++I)
 	{
@@ -844,8 +837,8 @@ DiffBlocksParallelT(FIOReader&				  BaseDataReader,
 
 	const uint64 BaseDataSize = BaseDataReader.GetSize();
 
-	HashSet<uint32, FIdentityHash32>							 SourceWeakHashSet;
-	HashSet<FGenericBlock, FBlockStrongHash, FBlockStrongHashEq> SourceStrongHashSet;
+	THashSet<uint32, FIdentityHash32>							  SourceWeakHashSet;
+	THashSet<FGenericBlock, FBlockStrongHash, FBlockStrongHashEq> SourceStrongHashSet;
 
 	for (uint32 I = 0; I < uint32(SourceBlocks.size()); ++I)
 	{
@@ -857,10 +850,10 @@ DiffBlocksParallelT(FIOReader&				  BaseDataReader,
 
 	struct FTask
 	{
-		uint64														 Offset = 0;
-		uint64														 Size	= 0;
-		std::vector<FHash128>										 Sequence;
-		HashSet<FGenericBlock, FBlockStrongHash, FBlockStrongHashEq> BaseStrongHashSet;
+		uint64														  Offset = 0;
+		uint64														  Size	 = 0;
+		std::vector<FHash128>										  Sequence;
+		THashSet<FGenericBlock, FBlockStrongHash, FBlockStrongHashEq> BaseStrongHashSet;
 	};
 
 	BytesPerTask = std::max<uint64>(BlockSize, BytesPerTask);
@@ -909,9 +902,9 @@ DiffBlocksParallelT(FIOReader&				  BaseDataReader,
 
 					const uint8* TaskEnd = TaskBuffer + Task.Size;
 
-					const uint32							 MaxWeakHashFalsePositives = 8;
-					HashMap<uint32, uint32, FIdentityHash32> WeakHashFalsePositives;
-					HashSet<uint32, FIdentityHash32>		 WeakHashBanList;
+					const uint32							  MaxWeakHashFalsePositives = 8;
+					THashMap<uint32, uint32, FIdentityHash32> WeakHashFalsePositives;
+					THashSet<uint32, FIdentityHash32>		  WeakHashBanList;
 
 					auto ScanFn = [&SourceWeakHashSet,
 								   &WeakHashBanList,
@@ -969,7 +962,7 @@ DiffBlocksParallelT(FIOReader&				  BaseDataReader,
 	BaseDataReader.FlushAll();
 	TaskGroup.wait();
 
-	HashSet<FGenericBlock, FBlockStrongHash, FBlockStrongHashEq> BaseStrongHashSet;
+	THashSet<FGenericBlock, FBlockStrongHash, FBlockStrongHashEq> BaseStrongHashSet;
 
 	for (FTask& Task : Tasks)
 	{
@@ -1083,7 +1076,7 @@ OptimizeNeedList(const std::vector<FNeedBlock>& Input, uint64 MaxMergedBlockSize
 		Result.push_back(Cmd);
 	}
 
-	std::sort(Result.begin(), Result.end(), [](const FCopyCommand& A, const FCopyCommand& B) { return A.SourceOffset < B.SourceOffset; });
+	std::sort(Result.begin(), Result.end(), FCopyCommand::FCompareBySourceOffset());
 
 	for (uint64 I = 1; I < Result.size(); ++I)
 	{
@@ -1118,12 +1111,11 @@ BuildTargetBuffer(const uint8*			 SourceData,
 				  const uint8*			 BaseData,
 				  uint64				 BaseSize,
 				  const FNeedList&		 NeedList,
-				  EStrongHashAlgorithmID StrongHasher,
-				  FProxyPool*			 ProxyPool)
+				  EStrongHashAlgorithmID StrongHasher)
 {
 	FMemReader SourceReader(SourceData, SourceSize);
 	FMemReader BaseReader(BaseData, BaseSize);
-	return BuildTargetBuffer(SourceReader, BaseReader, NeedList, StrongHasher, ProxyPool);
+	return BuildTargetBuffer(SourceReader, BaseReader, NeedList, StrongHasher);
 }
 
 struct FReadSchedule
@@ -1159,35 +1151,25 @@ BuildReadSchedule(const std::vector<FNeedBlock>& Blocks)
 class FBlockCache
 {
 public:
-	FBuffer						   BlockData;
-	HashMap<FHash128, FBufferView> BlockMap;  // Decompressed block data by hash
+	FBuffer							BlockData;
+	THashMap<FHash128, FBufferView> BlockMap;  // Decompressed block data by hash
 };
 
 FBuffer
-BuildTargetBuffer(FIOReader&			 SourceProvider,
-				  FIOReader&			 BaseProvider,
-				  const FNeedList&		 NeedList,
-				  EStrongHashAlgorithmID StrongHasher,
-				  FProxyPool*			 ProxyPool,
-				  FBlockCache*			 BlockCache)
+BuildTargetBuffer(FIOReader& SourceProvider, FIOReader& BaseProvider, const FNeedList& NeedList, EStrongHashAlgorithmID StrongHasher)
 {
 	FBuffer				Result;
 	const FNeedListSize SizeInfo = ComputeNeedListSize(NeedList);
 	Result.Resize(SizeInfo.TotalBytes);
 	FMemReaderWriter ResultWriter(Result.Data(), Result.Size());
-	BuildTarget(ResultWriter, SourceProvider, BaseProvider, NeedList, StrongHasher, ProxyPool, BlockCache);
+	BuildTarget(ResultWriter, SourceProvider, BaseProvider, NeedList, StrongHasher);
 	return Result;
 }
 
-enum EListType {
-	Source,
-	Base
-};
-
-inline void
-AddGlobalProgress(uint64 Size, EListType ListType)
+void
+AddGlobalProgress(uint64 Size, EBlockListType ListType)
 {
-	if (ListType == EListType::Source)
+	if (ListType == EBlockListType::Source)
 	{
 		Size *= GLOBAL_PROGRESS_SOURCE_SCALE;
 	}
@@ -1208,7 +1190,7 @@ DownloadBlocks(FProxyPool&					  ProxyPool,
 	const bool			bParentThreadVerbose = GLogVerbose;
 	std::atomic<uint64> NumActiveLogThreads	 = {};
 
-	HashSet<FHash128>			   DownloadedBlocks;
+	THashSet<FHash128>			   DownloadedBlocks;
 	std::vector<FNeedBlock>		   RemainingNeedBlocks;
 	const std::vector<FNeedBlock>* NeedBlocks = &OriginalUniqueNeedBlocks;
 
@@ -1350,7 +1332,8 @@ BuildTarget(FIOWriter&			   Result,
 			const FNeedList&	   NeedList,
 			EStrongHashAlgorithmID StrongHasher,
 			FProxyPool*			   ProxyPool,
-			FBlockCache*		   BlockCache)
+			FBlockCache*		   BlockCache,
+			FScavengeDatabase*	   ScavengeDatabase)
 {
 	UNSYNC_LOG_INDENT;
 
@@ -1374,21 +1357,40 @@ BuildTarget(FIOWriter&			   Result,
 
 	// Copy the need list as it will be progressively filtered by various sources:
 	// - block cache
+	// - scevenge from external local files
 	// - download from proxy
 	// - read from source file
 	// TODO: could use a bit array to indicate if the entry is valid instead of full copy
-	std::vector<FNeedBlock> FilteredSourceNeedList = NeedList.Source;
+	std::vector<FNeedBlock> FilteredSourceNeedList;
+
+	if (ScavengeDatabase)
+	{
+		THashSet<FHash128> ScavengedBlocks;
+
+		BuildTargetFromScavengedData(Result, NeedList.Source, *ScavengeDatabase, StrongHasher, ScavengedBlocks);
+
+		FilteredSourceNeedList = NeedList.Source;
+		auto FilterResult =
+			std::remove_if(FilteredSourceNeedList.begin(), FilteredSourceNeedList.end(), [&ScavengedBlocks](const FNeedBlock& Block) {
+				return ScavengedBlocks.find(Block.Hash.ToHash128()) != ScavengedBlocks.end();
+			});
+		FilteredSourceNeedList.erase(FilterResult, FilteredSourceNeedList.end());
+	}
+	else
+	{
+		FilteredSourceNeedList = NeedList.Source;
+	}
 
 	auto ProcessNeedList = [bAllowVerboseLog, LogIndent, &Result, &bGotError, &WriteSemaphore, &WriteTasks, &bWaitingForBaseData](
 							   FIOReader&					  DataProvider,
 							   const std::vector<FNeedBlock>& NeedBlocks,
 							   uint64						  TotalCopySize,
-							   EListType					  ListType,
+							   EBlockListType				  ListType,
 							   std::atomic<bool>&			  bCompletionFlag) {
 		FLogIndentScope IndentScope(LogIndent, true);
 		uint64			ReadBytesTotal = 0;
 
-		const wchar_t* ListName = ListType == EListType::Source ? L"source" : L"base";
+		const wchar_t* ListName = ListType == EBlockListType::Source ? L"source" : L"base";
 
 		if (!DataProvider.IsValid())
 		{
@@ -1421,7 +1423,7 @@ BuildTarget(FIOWriter&			   Result,
 
 		while (!ReadSchedule.Requests.empty() && !bGotError)
 		{
-			uint64 BlockIndex;
+			uint64 BlockIndex = ~0ull;
 
 			uint32 LargeRequestSlot = ~0u;
 			for (uint32 I = 0; I < MaxActiveLargeRequests; ++I)
@@ -1491,7 +1493,8 @@ BuildTarget(FIOWriter&			   Result,
 
 			ReadBytesTotal += ReadBytes;
 
-			FLogVerbosityScope VerbosityScope(GLogVerbose || (bAllowVerboseLog && (ListType == EListType::Base && bWaitingForBaseData)));
+			FLogVerbosityScope VerbosityScope(GLogVerbose ||
+											  (bAllowVerboseLog && (ListType == EBlockListType::Base && bWaitingForBaseData)));
 
 			ProgressLogger.Add(ReadBytes);
 
@@ -1515,11 +1518,11 @@ BuildTarget(FIOWriter&			   Result,
 
 		if (SizeInfo.BaseBytes)
 		{
-			ProcessNeedList(Base, NeedList.Base, SizeInfo.BaseBytes, EListType::Base, bBaseDataCopyTaskDone);
+			ProcessNeedList(Base, NeedList.Base, SizeInfo.BaseBytes, EBlockListType::Base, bBaseDataCopyTaskDone);
 		}
 	});
 
-	HashSet<FHash128> CachedBlocks;
+	THashSet<FHash128> CachedBlocks;
 
 	if (BlockCache)
 	{
@@ -1534,7 +1537,7 @@ BuildTarget(FIOWriter&			   Result,
 				FBufferView BlockBuffer = It->second;
 
 				Result.Write(BlockBuffer.Data, NeedBlock.TargetOffset, BlockBuffer.Size);
-				AddGlobalProgress(NeedBlock.Size, EListType::Source);
+				AddGlobalProgress(NeedBlock.Size, EBlockListType::Source);
 			}
 		}
 
@@ -1552,16 +1555,19 @@ BuildTarget(FIOWriter&			   Result,
 	}
 	else if (ProxyPool && ProxyPool->IsValid())
 	{
+		UNSYNC_VERBOSE(L"Downloading blocks");
+		UNSYNC_LOG_INDENT;
+
 		struct FBlockWriteCmd
 		{
 			uint64 TargetOffset = 0;
 			uint64 Size			= 0;
 		};
-		std::mutex							   DownloadedBlocksMutex;
-		HashSet<FHash128>					   DownloadedBlocks;
-		HashMap<FHash128, FBlockWriteCmd>	   BlockWriteMap;
-		HashMap<FHash128, std::vector<uint64>> BlockScatterMap;
-		std::vector<FNeedBlock>				   UniqueNeedList;
+		std::mutex								DownloadedBlocksMutex;
+		THashSet<FHash128>						DownloadedBlocks;
+		THashMap<FHash128, FBlockWriteCmd>		BlockWriteMap;
+		THashMap<FHash128, std::vector<uint64>> BlockScatterMap;
+		std::vector<FNeedBlock>					UniqueNeedList;
 
 		uint64 EstimatedDownloadSize = 0;
 		for (const FNeedBlock& Block : FilteredSourceNeedList)
@@ -1582,14 +1588,13 @@ BuildTarget(FIOWriter&			   Result,
 			}
 		}
 
-		UNSYNC_VERBOSE(L"Downloading blocks");
-		UNSYNC_LOG_INDENT;
-
 		FTaskGroup DecompressTasks;
 
 		FLogProgressScope DownloadProgressLogger(EstimatedDownloadSize, ELogProgressUnits::MB);
 
-		FSemaphore DecompressionSemaphore(256);	 // limit how many decompression tasks can be queued up to avoid memory bloat
+		// limit how many decompression tasks can be queued up to avoid memory bloat
+		const uint64 MaxConcurrentDecompressionTasks = 64;
+		FSemaphore DecompressionSemaphore(MaxConcurrentDecompressionTasks);
 
 		const bool			bParentThreadVerbose = GLogVerbose;
 		const uint32		ParentThreadIndent	 = GLogIndent;
@@ -1680,7 +1685,7 @@ BuildTarget(FIOWriter&			   Result,
 								if (bOk)
 								{
 									uint64 WrittenBytes = Result.Write(DecompressedData.GetData(), Cmd.TargetOffset, Cmd.Size);
-									AddGlobalProgress(Cmd.Size, EListType::Source);
+									AddGlobalProgress(Cmd.Size, EBlockListType::Source);
 
 									if (WrittenBytes != Cmd.Size)
 									{
@@ -1717,7 +1722,7 @@ BuildTarget(FIOWriter&			   Result,
 									}
 								}
 
-								AddGlobalProgress(Cmd.Size * ScatterList.size(), EListType::Source);
+								AddGlobalProgress(Cmd.Size * ScatterList.size(), EBlockListType::Source);
 							}
 						}
 
@@ -1760,7 +1765,7 @@ BuildTarget(FIOWriter&			   Result,
 
 			UNSYNC_LOG_INDENT;
 
-			ProcessNeedList(Source, FilteredSourceNeedList, RemainingBytes, EListType::Source, bSourceDataCopyTaskDone);
+			ProcessNeedList(Source, FilteredSourceNeedList, RemainingBytes, EBlockListType::Source, bSourceDataCopyTaskDone);
 		}
 		else
 		{
@@ -1769,7 +1774,7 @@ BuildTarget(FIOWriter&			   Result,
 	}
 	else if (SizeInfo.SourceBytes)
 	{
-		ProcessNeedList(Source, FilteredSourceNeedList, SizeInfo.SourceBytes, EListType::Source, bSourceDataCopyTaskDone);
+		ProcessNeedList(Source, FilteredSourceNeedList, SizeInfo.SourceBytes, EBlockListType::Source, bSourceDataCopyTaskDone);
 	}
 
 	if (!bBaseDataCopyTaskDone)
@@ -2353,7 +2358,7 @@ ValidateTarget(FIOReader& Reader, const FNeedList& NeedList, EStrongHashAlgorith
 					FGenericHash		 Hash = ComputeHash(CmdBuffer->GetData() + BlockBufferOffset, Block.Size, StrongHasher);
 					if (Hash != Block.HashStrong)
 					{
-						UNSYNC_ERROR(L"Found block hash mismatch");
+						UNSYNC_ERROR(L"Found block hash mismatch at offset %llu", llu(BlockBufferOffset));
 						NumInvalidBlocks++;
 						return;
 					}
@@ -2469,21 +2474,22 @@ SyncFile(const FNeedList&		   NeedList,
 			}
 		}
 
+		LogStatus(TargetFilePath.wstring().c_str(), L"Patching");
+
 		FDeferredOpenReader SourceFile([SourceFilePath, TargetFilePath] {
 			UNSYNC_VERBOSE(L"Opening source file '%ls'", SourceFilePath.wstring().c_str());
 			LogStatus(TargetFilePath.wstring().c_str(), L"Opening source file");
-
 			return std::unique_ptr<NativeFile>(new NativeFile(SourceFilePath, EFileMode::ReadOnlyUnbuffered));
 		});
 
-		LogStatus(TargetFilePath.wstring().c_str(), L"Patching");
 		BuildTarget(*TargetFile,
 					SourceFile,
 					BaseDataReader,
 					NeedList,
 					Options.Algorithm.StrongHashAlgorithmId,
 					Options.ProxyPool,
-					Options.BlockCache);
+					Options.BlockCache,
+					Options.ScavengeDatabase);
 
 		Result.SourceBytes = ComputeSize(NeedList.Source);
 		Result.BaseBytes   = ComputeSize(NeedList.Base);
@@ -3031,24 +3037,6 @@ GetAnonymizedHostName()
 	return Result;
 }
 
-struct FFileSyncTask
-{
-	const FFileManifest* SourceManifest = nullptr;
-	const FFileManifest* BaseManifest	= nullptr;
-	FPath				 OriginalSourceFilePath;
-	FPath				 ResolvedSourceFilePath;
-	FPath				 BaseFilePath;
-	FPath				 TargetFilePath;
-	FPath				 RelativeFilePath;
-	FNeedList			 NeedList;
-
-	uint64 NeedBytesFromSource = 0;
-	uint64 NeedBytesFromBase   = 0;
-	uint64 TotalSizeBytes	   = 0;
-
-	bool IsBaseValid() const { return !BaseFilePath.empty(); }
-};
-
 struct FFileSyncTaskBatch
 {
 	std::vector<const FFileSyncTask*> FileTasks;
@@ -3064,7 +3052,7 @@ struct FFileSyncTaskBatch
 
 		uint64 OutputCursor = 0;
 
-		HashSet<FHash128>		UniqueBlockSet;
+		THashSet<FHash128>		UniqueBlockSet;
 		std::vector<FNeedBlock> UniqueNeedBlocks;
 		for (const FFileSyncTask* Task : FileTasks)
 		{
@@ -3572,6 +3560,22 @@ SyncDirectory(const FSyncDirectoryOptions& SyncOptions)
 	GGlobalProgressTotal =
 		EstimatedNeedBytesFromSource * GLOBAL_PROGRESS_SOURCE_SCALE + EstimatedNeedBytesFromBase * GLOBAL_PROGRESS_BASE_SCALE;
 
+	std::unique_ptr<FScavengeDatabase> ScavengeDatabase;
+	if (!SyncOptions.ScavengeRoot.empty())
+	{
+		UNSYNC_VERBOSE(L"Scavenging blocks from existing data sets");
+		UNSYNC_LOG_INDENT;
+
+		FTimePoint ScavengeDbTimeBegin = TimePointNow();
+
+		TArrayView<FFileSyncTask> AllFileTasksView = MakeView(AllFileTasks.data(), AllFileTasks.size());
+		ScavengeDatabase = std::unique_ptr<FScavengeDatabase>(FScavengeDatabase::BuildFromFileSyncTasks(SyncOptions, AllFileTasksView));
+
+		double Duration = DurationSec(ScavengeDbTimeBegin, TimePointNow());
+
+		UNSYNC_VERBOSE(L"Done in %.3f sec", Duration);
+	}
+
 	LogGlobalProgress();
 
 	if (ProxySettings.IsValid() && ProxyPool.IsValid())
@@ -3694,6 +3698,7 @@ SyncDirectory(const FSyncDirectoryOptions& SyncOptions)
 								 &BackgroundTaskStatMutex,
 								 &BackgroundTaskResults,
 								 &NumFailedTasks,
+								 &ScavengeDatabase,
 								 &ProxyPool](const FFileSyncTask& Item, FBlockCache* BlockCache, bool bBackground) {
 			UNSYNC_VERBOSE(L"Copy '%ls' (%ls)", Item.TargetFilePath.wstring().c_str(), (Item.NeedBytesFromBase) ? L"partial" : L"full");
 
@@ -3711,6 +3716,7 @@ SyncDirectory(const FSyncDirectoryOptions& SyncOptions)
 			SyncFileOptions.BlockSize			 = SourceBlockSize;
 			SyncFileOptions.ProxyPool			 = &ProxyPool;
 			SyncFileOptions.BlockCache			 = BlockCache;
+			SyncFileOptions.ScavengeDatabase	 = ScavengeDatabase.get();
 			SyncFileOptions.bValidateTargetFiles = SyncOptions.bValidateTargetFiles;
 
 			FFileSyncResult SyncResult =
@@ -4496,10 +4502,10 @@ LogManifestFiles(ELogLevel LogLevel, const FDirectoryManifest& Manifest)
 	}
 }
 
-HashMap<FGenericHash, FGenericBlock>
+THashMap<FGenericHash, FGenericBlock>
 BuildBlockMap(const FDirectoryManifest& Manifest, bool bNeedMacroBlocks)
 {
-	HashMap<FGenericHash, FGenericBlock> Result;
+	THashMap<FGenericHash, FGenericBlock> Result;
 	for (const auto& It : Manifest.Files)
 	{
 		const FFileManifest& File = It.second;
@@ -4524,11 +4530,11 @@ BuildBlockMap(const FDirectoryManifest& Manifest, bool bNeedMacroBlocks)
 void
 LogManifestDiff(ELogLevel LogLevel, const FDirectoryManifest& ManifestA, const FDirectoryManifest& ManifestB)
 {
-	HashMap<FGenericHash, FGenericBlock> BlocksA = BuildBlockMap(ManifestA, false);
-	HashMap<FGenericHash, FGenericBlock> BlocksB = BuildBlockMap(ManifestB, false);
+	THashMap<FGenericHash, FGenericBlock> BlocksA = BuildBlockMap(ManifestA, false);
+	THashMap<FGenericHash, FGenericBlock> BlocksB = BuildBlockMap(ManifestB, false);
 
-	HashMap<FGenericHash, FGenericBlock> MacroBlocksA = BuildBlockMap(ManifestA, true);
-	HashMap<FGenericHash, FGenericBlock> MacroBlocksB = BuildBlockMap(ManifestB, true);
+	THashMap<FGenericHash, FGenericBlock> MacroBlocksA = BuildBlockMap(ManifestA, true);
+	THashMap<FGenericHash, FGenericBlock> MacroBlocksB = BuildBlockMap(ManifestB, true);
 
 	uint32 NumCommonBlocks		= 0;
 	uint64 TotalCommonBlockSize = 0;
diff --git a/Engine/Source/Programs/Unsync/Private/UnsyncCore.h b/Engine/Source/Programs/Unsync/Private/UnsyncCore.h
--- a/Engine/Source/Programs/Unsync/Private/UnsyncCore.h
+++ b/Engine/Source/Programs/Unsync/Private/UnsyncCore.h
@@ -21,6 +21,7 @@ extern bool GDryRun;
 class FProxy;
 class FProxyPool;
 class FBlockCache;
+class FScavengeDatabase;
 struct FRemoteDesc;
 struct FIOReader;
 struct FIOWriter;
@@ -59,6 +60,11 @@ struct FCopyCommand
 	uint64 Size			= 0;
 	uint64 SourceOffset = 0;
 	uint64 TargetOffset = 0;
+
+	struct FCompareBySourceOffset
+	{
+		bool operator()(const FCopyCommand& A, const FCopyCommand& B) const { return A.SourceOffset < B.SourceOffset; }
+	};
 };
 
 struct FNeedBlock : FCopyCommand
@@ -250,22 +256,20 @@ void BuildTarget(FIOWriter&				Result,
 				 const FNeedList&		NeedList,
 				 EStrongHashAlgorithmID StrongHasher,
 				 FProxyPool*			ProxyPool = nullptr,
-				 FBlockCache*			BlockCache = nullptr);
+				 FBlockCache*			BlockCache = nullptr,
+				 FScavengeDatabase*		ScavengeDatabase = nullptr);
 
 FBuffer BuildTargetBuffer(FIOReader&			 SourceProvider,
 						  FIOReader&			 BaseProvider,
 						  const FNeedList&		 NeedList,
-						  EStrongHashAlgorithmID StrongHasher,
-						  FProxyPool*			 ProxyPool = nullptr,
-						  FBlockCache*			 BlockCache = nullptr);
+						  EStrongHashAlgorithmID StrongHasher);
 
 FBuffer BuildTargetBuffer(const uint8*			 SourceData,
 						  uint64				 SourceSize,
 						  const uint8*			 BaseData,
 						  uint64				 BaseSize,
 						  const FNeedList&		 NeedList,
-						  EStrongHashAlgorithmID StrongHasher,
-						  FProxyPool*			 ProxyPool = nullptr);
+						  EStrongHashAlgorithmID StrongHasher);
 
 FBuffer BuildTargetWithPatch(const uint8* PatchData, uint64 PatchSize, const uint8* BaseData, uint64 BaseSize);
 
@@ -292,6 +296,24 @@ enum class EFileSyncStatus
 
 const wchar_t* ToString(EFileSyncStatus Status);
 
+struct FFileSyncTask
+{
+	const FFileManifest* SourceManifest = nullptr;
+	const FFileManifest* BaseManifest	= nullptr;
+	FPath				 OriginalSourceFilePath;
+	FPath				 ResolvedSourceFilePath;
+	FPath				 BaseFilePath;
+	FPath				 TargetFilePath;
+	FPath				 RelativeFilePath;
+	FNeedList			 NeedList;
+
+	uint64 NeedBytesFromSource = 0;
+	uint64 NeedBytesFromBase   = 0;
+	uint64 TotalSizeBytes	   = 0;
+
+	bool IsBaseValid() const { return !BaseFilePath.empty(); }
+};
+
 struct FFileSyncResult
 {
 	EFileSyncStatus Status			= EFileSyncStatus::ErrorUnknown;
@@ -310,6 +332,7 @@ struct FSyncFileOptions
 
 	FProxyPool* ProxyPool = nullptr;
 	FBlockCache* BlockCache = nullptr;
+	FScavengeDatabase* ScavengeDatabase = nullptr;
 
 	bool bValidateTargetFiles = true;  // WARNING: turning this off is intended only for testing/profiling
 };
@@ -356,12 +379,21 @@ enum class ESyncSourceType
 	Server,
 };
 
+enum EBlockListType {
+	Source,
+	Base
+};
+
+void AddGlobalProgress(uint64 Size, EBlockListType ListType);
+
 struct FSyncDirectoryOptions
 {
 	ESyncSourceType	   SourceType;
 	FPath			   Source;	  // remote data location
 	FPath			   Target;	  // output target location
 	FPath			   Base;	  // base data location, which typically is the same as sync target
+	FPath			   ScavengeRoot; // base directory where we may want to find reusable blocks
+	uint32			   ScavengeDepth = 5; // how deep to look for unsync manifests
 	std::vector<FPath> Overlays;  // extra source directories to overlay over primary (add extra files, replace existing files)
 	FPath			   SourceManifestOverride;	// force the manifest to be read from a specified file instead of source directory
 	FSyncFilter*	   SyncFilter = nullptr;	// filter callback for partial sync support
diff --git a/Engine/Source/Programs/Unsync/Private/UnsyncFile.cpp b/Engine/Source/Programs/Unsync/Private/UnsyncFile.cpp
--- a/Engine/Source/Programs/Unsync/Private/UnsyncFile.cpp
+++ b/Engine/Source/Programs/Unsync/Private/UnsyncFile.cpp
@@ -25,9 +25,9 @@ namespace unsync {
 bool GForceBufferedFiles = false;
 
 // Returns extended absolute path of a form \\?\D:\verylongpath or \\?\UNC\servername\verylongpath
-// Expects an absolute path input.
+// Expects an absolute path input. Returns original path on non-Windows.
 // https://docs.microsoft.com/en-us/windows/win32/fileio/maximum-file-path-limitation
-static FPath
+FPath
 MakeExtendedAbsolutePath(const FPath& InAbsolutePath)
 {
 	if (InAbsolutePath.empty())
diff --git a/Engine/Source/Programs/Unsync/Private/UnsyncFile.h b/Engine/Source/Programs/Unsync/Private/UnsyncFile.h
--- a/Engine/Source/Programs/Unsync/Private/UnsyncFile.h
+++ b/Engine/Source/Programs/Unsync/Private/UnsyncFile.h
@@ -464,4 +464,9 @@ uint64 ToWindowsFileTime(const std::filesystem::file_time_type& T);
 struct FSyncFilter;
 FFileAttributeCache CreateFileAttributeCache(const FPath& Root, const FSyncFilter* SyncFilter = nullptr);
 
+// Returns extended absolute path of a form \\?\D:\verylongpath or \\?\UNC\servername\verylongpath
+// Expects an absolute path input. Returns original path on non-Windows.
+// https://docs.microsoft.com/en-us/windows/win32/fileio/maximum-file-path-limitation
+FPath MakeExtendedAbsolutePath(const FPath& InAbsolutePath);
+
 }  // namespace unsync
diff --git a/Engine/Source/Programs/Unsync/Private/UnsyncHashTable.h b/Engine/Source/Programs/Unsync/Private/UnsyncHashTable.h
new file mode 100644
--- /dev/null
+++ b/Engine/Source/Programs/Unsync/Private/UnsyncHashTable.h
@@ -0,0 +1,19 @@
+// Copyright Epic Games, Inc. All Rights Reserved.
+
+#pragma once
+
+UNSYNC_THIRD_PARTY_INCLUDES_START
+#include <flat_hash_map.hpp>
+UNSYNC_THIRD_PARTY_INCLUDES_END
+
+namespace unsync {
+
+template<typename K, typename V, typename Hash = std::hash<K>, typename Eq = std::equal_to<K>>
+// using THashMap = std::unordered_map<K,V,Hash,Eq>;
+using THashMap = ska::flat_hash_map<K, V, Hash, Eq>;
+
+template<typename K, typename Hash = std::hash<K>, typename Eq = std::equal_to<K>>
+// using THashSet = std::unordered_set<K, Hash, Eq>;
+using THashSet = ska::flat_hash_set<K, Hash, Eq>;
+
+}  // namespace unsync
diff --git a/Engine/Source/Programs/Unsync/Private/UnsyncMain.cpp b/Engine/Source/Programs/Unsync/Private/UnsyncMain.cpp
--- a/Engine/Source/Programs/Unsync/Private/UnsyncMain.cpp
+++ b/Engine/Source/Programs/Unsync/Private/UnsyncMain.cpp
@@ -69,6 +69,7 @@ InnerMain(int Argc, char** Argv)
 	std::string				 ProtocolName = "jupiter";
 	std::string				 HttpHeaderFilenameUtf8;
 	std::string				 QueryStringUtf8;
+	std::string				 ScavengeRootUtf8;
 	bool					 bForceOperation	 = false;
 	bool					 bAllowInsecureTls	 = false;
 	bool					 bUseTls			 = false;
@@ -217,6 +218,9 @@ InnerMain(int Argc, char** Argv)
 	SubSync->add_flag("--no-output-validation", bNoOutputValidation, "Skip final patched file block hash validation (DANGEROUS)");
 	SubSync->add_flag("--no-space-validation", bNoSpaceValidation, "Skip checking available disk space before sync (DANGEROUS)");
 	SubSync->add_option("-b, --block", HashOrSyncBlockSize, "Block size in bytes (default=64KB)");
+
+	SubSync->add_option("--scavenge", ScavengeRootUtf8, "Search for unsync manifests and reusable blocks in this directory (EXPERIMENTAL)");
+
 	SubCommands.push_back(SubSync);
 
 	CLI::App* SubPatch = Cli.add_subcommand("patch", "Applies a patch generated with 'diff' on top of base file");
@@ -394,6 +398,7 @@ InnerMain(int Argc, char** Argv)
 	FPath SourceFilename		 = NormalizeFilenameUtf8(SourceFilenameUtf8);
 	FPath TargetFilename		 = NormalizeFilenameUtf8(TargetFilenameUtf8);
 	FPath PatchFilename			 = NormalizeFilenameUtf8(PatchFilenameUtf8);
+	FPath ScavengeRoot			 = NormalizeFilenameUtf8(ScavengeRootUtf8);
 	FPath SourceManifestFilename = NormalizeFilenameUtf8(SourceManifestFilenameUtf8);
 
 	UNSYNC_VERBOSE(L"UNSYNC %hs", GetVersionString().c_str());
@@ -657,6 +662,7 @@ InnerMain(int Argc, char** Argv)
 		SyncOptions.Filter				   = &SyncFilter;
 		SyncOptions.bValidateTargetFiles   = !bNoOutputValidation;
 		SyncOptions.bCheckAvailableSpace   = !bNoSpaceValidation;
+		SyncOptions.ScavengeRoot		   = ScavengeRoot;
 
 		for (const std::string& Entry : OverlayArrayUtf8)
 		{
diff --git a/Engine/Source/Programs/Unsync/Private/UnsyncScavenger.cpp b/Engine/Source/Programs/Unsync/Private/UnsyncScavenger.cpp
new file mode 100644
--- /dev/null
+++ b/Engine/Source/Programs/Unsync/Private/UnsyncScavenger.cpp
@@ -0,0 +1,451 @@
+// Copyright Epic Games, Inc. All Rights Reserved.
+
+#include "UnsyncScavenger.h"
+#include "UnsyncFile.h"
+#include "UnsyncLog.h"
+#include "UnsyncSerialization.h"
+#include "UnsyncThread.h"
+namespace unsync {
+
+FScavengeDatabase*
+FScavengeDatabase::BuildFromFileSyncTasks(const FSyncDirectoryOptions& SyncOptions, TArrayView<FFileSyncTask> AllFileTasks)
+{
+	FScavengeDatabase* Result = new FScavengeDatabase;
+
+	THashSet<FHash128> NeededBlocks;
+	for (const FFileSyncTask& FileTask : AllFileTasks)
+	{
+		for (const FNeedBlock& SourceNeedBlock : FileTask.NeedList.Source)
+		{
+			NeededBlocks.insert(SourceNeedBlock.Hash.ToHash128());
+		}
+	}
+
+	if (NeededBlocks.empty())
+	{
+		return Result;
+	}
+
+	const FPath ExtendedTargetPath = MakeExtendedAbsolutePath(SyncOptions.Target);
+
+	UNSYNC_VERBOSE(L"Scanning '%ls' for usable manifests", SyncOptions.ScavengeRoot.c_str());
+
+	for (std::filesystem::recursive_directory_iterator DirIt = RecursiveDirectoryScan(SyncOptions.ScavengeRoot);
+		 DirIt != std::filesystem::recursive_directory_iterator();
+		 ++DirIt)
+	{
+		const std::filesystem::directory_entry& Dir = *DirIt;
+
+		if (!Dir.is_directory())
+		{
+			continue;
+		}
+
+		const int32 DirDepth = DirIt.depth();
+		if (DirDepth + 1 >= int32(SyncOptions.ScavengeDepth))
+		{
+			DirIt.disable_recursion_pending();
+		}
+
+		const FPath& DirPath = Dir.path();
+		if (DirPath == ExtendedTargetPath)
+		{
+			// Exclude the current sync target path from scavanging database
+			DirIt.disable_recursion_pending();
+			continue;
+		}
+
+		FPath DirStem = DirPath.stem();
+		if (!DirStem.compare(L".unsync"))
+		{
+			DirIt.disable_recursion_pending();
+
+			FPath ManifestPath = DirPath / "manifest.bin";
+			if (PathExists(ManifestPath))
+			{
+				FScavengedManifest Entry;
+				Entry.ManifestPath = ManifestPath;
+				Entry.Root		   = DirPath.parent_path();
+				Result->Manifests.push_back(std::move(Entry));
+			}
+		}
+	}
+
+	UNSYNC_VERBOSE(L"Loading scavenged manifests: %llu", llu(Result->Manifests.size()));
+
+	ParallelForEach(Result->Manifests.begin(), Result->Manifests.end(), [](FScavengedManifest& Entry) {
+		FLogVerbosityScope VerbosityScope(false);  // turn off logging from threads
+		Entry.bValid = LoadDirectoryManifest(Entry.Manifest, Entry.Root, Entry.ManifestPath);
+		if (Entry.bValid)
+		{
+			Entry.FileList.reserve(Entry.Manifest.Files.size());
+			for (const auto& It : Entry.Manifest.Files)
+			{
+				const FPath& RelativeFileName = It.first;
+				Entry.FileList.push_back(RelativeFileName);
+			}
+		}
+	});
+
+	THashMap<FScavengeBlockSource, uint64, FScavengeBlockSource::FHash> BlockSourceUseCounts;
+
+	uint32 ManifestIndex = 0;
+	for (const FScavengedManifest& Entry : Result->Manifests)
+	{
+		uint32 FileIndex = 0;
+		for (const auto& FileManifestIt : Entry.Manifest.Files)
+		{
+			const FFileManifest& FileManifest = FileManifestIt.second;
+			for (const FGenericBlock& BlockInfo : FileManifest.Blocks)
+			{
+				FHash128 BlockHash = BlockInfo.HashStrong.ToHash128();
+				Result->UniqueBlockHashes.insert(BlockHash);
+				if (NeededBlocks.find(BlockHash) != NeededBlocks.end())
+				{
+					FScavengeBlockSource BlockSource;
+					BlockSource.Data.ManifestIndex = ManifestIndex;
+					BlockSource.Data.FileIndex	   = FileIndex;
+					Result->BlockMap.insert(std::make_pair(BlockHash, BlockSource));
+					Result->UniqueUsableBlockHashes.insert(BlockHash);
+
+					BlockSourceUseCounts[BlockSource] += 1;
+				}
+			}
+			++FileIndex;
+		}
+		++ManifestIndex;
+	}
+
+	struct FBlockSourceAndUseCount
+	{
+		FScavengeBlockSource Source;
+		uint64				 Count;
+	};
+
+	std::vector<FBlockSourceAndUseCount> SortedBlockSources;
+	for (const auto& It : BlockSourceUseCounts)
+	{
+		FBlockSourceAndUseCount Entry;
+		Entry.Source = It.first;
+		Entry.Count	 = It.second;
+		SortedBlockSources.push_back(Entry);
+	}
+
+	std::sort(SortedBlockSources.begin(), SortedBlockSources.end(), [](const FBlockSourceAndUseCount& A, const FBlockSourceAndUseCount& B) {
+		return A.Count > B.Count;
+	});
+
+	UNSYNC_VERBOSE(L"Found potentially useful files: %llu", llu(BlockSourceUseCounts.size()));
+
+#if 0  // Extra verbose scavenging status report
+	if (SortedBlockSources.size())
+	{
+		UNSYNC_LOG_INDENT;
+		UNSYNC_VERBOSE(L"Top useful files:");
+		for (size_t I = 0; I < SortedBlockSources.size() && I < 100; ++I)
+		{
+			const FBlockSourceAndUseCount& Item				= SortedBlockSources[I];
+			const FScavengedManifest&	   Manifest			= Result->Manifests[Item.Source.Data.ManifestIndex];
+			const FPath&				   RelativeFilePath = Manifest.FileList[Item.Source.Data.FileIndex];
+
+			FPath FilePath = Manifest.Root / RelativeFilePath;
+
+			UNSYNC_VERBOSE(L"- %s: %llu", FilePath.wstring().c_str(), Item.Count);
+		}
+	}
+#endif
+
+	UNSYNC_VERBOSE(L"Found potential usable blocks: %llu out of %llu needed (%.2f %%)",
+				   llu(Result->UniqueUsableBlockHashes.size()),
+				   llu(NeededBlocks.size()),
+				   100.0 * double(Result->UniqueUsableBlockHashes.size()) / double(NeededBlocks.size()));
+
+	return Result;
+}
+
+const FPath&
+FScavengeDatabase::GetPartialSourceFilePath(FScavengeBlockSource SourceId) const
+{
+	const FScavengedManifest& Manifest	   = Manifests[SourceId.Data.ManifestIndex];
+	const FPath&			  RelativePath = Manifest.FileList[SourceId.Data.FileIndex];
+	return RelativePath;
+}
+
+FPath
+FScavengeDatabase::GetFullSourceFilePath(FScavengeBlockSource SourceId) const
+{
+	const FScavengedManifest& Manifest	   = Manifests[SourceId.Data.ManifestIndex];
+	const FPath&			  RelativePath = Manifest.FileList[SourceId.Data.FileIndex];
+	return Manifest.Root / RelativePath;
+}
+
+const FFileManifest&
+FScavengeDatabase::GetFileManifest(FScavengeBlockSource SourceId) const
+{
+	const FScavengedManifest& ScavengedManifest = GetScavengedManifest(SourceId);
+	auto					  FileIt = ScavengedManifest.Manifest.Files.find(ScavengedManifest.FileList[SourceId.Data.FileIndex].wstring());
+	return FileIt->second;
+}
+
+bool
+FScavengeDatabase::IsSourceValid(FScavengeBlockSource SourceId) const
+{
+	if (SourceId.Data.ManifestIndex >= Manifests.size())
+	{
+		return false;
+	}
+
+	const FScavengedManifest& ScavangedManifest = Manifests[SourceId.Data.ManifestIndex];
+	if (SourceId.Data.FileIndex >= ScavangedManifest.FileList.size())
+	{
+		return false;
+	}
+
+	const FPath& RelativeFilePath = ScavangedManifest.FileList[SourceId.Data.FileIndex];
+	FPath		 FullFilePath	  = ScavangedManifest.Root / RelativeFilePath;
+
+	auto FileIt = ScavangedManifest.Manifest.Files.find(RelativeFilePath.wstring());
+	if (FileIt == ScavangedManifest.Manifest.Files.end())
+	{
+		return false;
+	}
+
+	const FFileManifest& FileManifest = FileIt->second;
+
+	FileAttributes Attrib = GetFileAttrib(FullFilePath);
+
+	return Attrib.bValid && Attrib.Mtime == FileManifest.Mtime && Attrib.Size == FileManifest.Size;
+}
+
+struct FCopyCommandWithBlockRange : FCopyCommand
+{
+	TArrayView<FNeedBlock> BlockRange;
+};
+
+// Similar to OptimizeNeedList, but assumes maintains commands in the same order as input blocks.
+// Input blocks are preferred to be sorted by source offset, but it is not a hard requirement.
+std::vector<FCopyCommandWithBlockRange> static OptimizeNeedListWithBlockRange(const std::vector<FNeedBlock>& Input,
+																			  uint64						 MaxMergedBlockSize)
+{
+	std::vector<FCopyCommandWithBlockRange> Result;
+	Result.reserve(Input.size());
+	for (const FNeedBlock& Block : Input)
+	{
+		FCopyCommandWithBlockRange Cmd;
+		Cmd.SourceOffset = Block.SourceOffset;
+		Cmd.TargetOffset = Block.TargetOffset;
+		Cmd.Size		 = Block.Size;
+		Cmd.BlockRange	 = MakeView(&Block, 1);
+		Result.push_back(Cmd);
+	}
+
+	for (uint64 I = 1; I < Result.size(); ++I)
+	{
+		FCopyCommandWithBlockRange& PrevBlock = Result[I - 1];
+		FCopyCommandWithBlockRange& ThisBlock = Result[I];
+		if (PrevBlock.SourceOffset + PrevBlock.Size == ThisBlock.SourceOffset &&
+			PrevBlock.TargetOffset + PrevBlock.Size == ThisBlock.TargetOffset && PrevBlock.Size + ThisBlock.Size <= MaxMergedBlockSize)
+		{
+			ThisBlock.SourceOffset = PrevBlock.SourceOffset;
+			ThisBlock.TargetOffset = PrevBlock.TargetOffset;
+			ThisBlock.Size += PrevBlock.Size;
+
+			UNSYNC_ASSERT(PrevBlock.BlockRange.EndPtr == ThisBlock.BlockRange.BeginPtr);
+			UNSYNC_ASSERT(ThisBlock.Size <= MaxMergedBlockSize);
+
+			ThisBlock.BlockRange.BeginPtr = PrevBlock.BlockRange.BeginPtr;
+
+			// Invalidate previous block
+			PrevBlock.BlockRange.EndPtr = PrevBlock.BlockRange.BeginPtr;
+			PrevBlock.Size				= 0;
+		}
+	}
+
+	for (uint64 I = 0; I < Result.size(); ++I)
+	{
+		UNSYNC_ASSERT(Result[I].Size <= MaxMergedBlockSize);
+	}
+
+	auto It = std::remove_if(Result.begin(), Result.end(), [](const FCopyCommand& Block) { return Block.Size == 0; });
+
+	Result.erase(It, Result.end());
+
+	{
+		uint64 BlockCount = 0;
+		for (uint64 I = 0; I < Result.size(); ++I)
+		{
+			BlockCount += Result[I].BlockRange.Size();
+		}
+		UNSYNC_ASSERT(BlockCount == Input.size());
+	}
+
+	return Result;
+}
+
+void
+BuildTargetFromScavengedData(FIOWriter&						Result,
+							 const std::vector<FNeedBlock>& NeedList,
+							 const FScavengeDatabase&		ScavengeDatabase,
+							 EStrongHashAlgorithmID			StrongHasher,
+							 THashSet<FHash128>&			OutScavengedBlocks)
+{
+	std::vector<FNeedBlock> ScavengeNeedList;
+
+	THashMap<FScavengeBlockSource, uint64, FScavengeBlockSource::FHash> PossibleSources;
+
+	const FScavengeBlockMap& ScavengeBlockMap = ScavengeDatabase.GetBlockMap();
+	uint64					 TotalNeedSize	  = 0;
+	for (const FNeedBlock& SourceNeedBlock : NeedList)
+	{
+		FHash128   NeedBlockHash = SourceNeedBlock.Hash.ToHash128();
+		TotalNeedSize += SourceNeedBlock.Size;
+		const auto Sources		 = ScavengeBlockMap.equal_range(NeedBlockHash);
+		if (Sources.first != Sources.second)
+		{
+			ScavengeNeedList.push_back(SourceNeedBlock);
+			for (auto SourceIt = Sources.first; SourceIt != Sources.second; ++SourceIt)
+			{
+				const FScavengeBlockSource& BlockSource = SourceIt->second;
+				PossibleSources[BlockSource] += 1;
+			}
+		}
+	}
+
+	struct FPossibleSource : FScavengeBlockSource
+	{
+		uint64 NumHits = 0;
+		FPath  FileName;
+		FPath  FullSourceFilePath;
+	};
+
+	std::vector<FPossibleSource> SortedPossibleSources;
+
+	for (const auto& It : PossibleSources)
+	{
+		if (ScavengeDatabase.IsSourceValid(It.first))
+		{
+			FPossibleSource Entry;
+			Entry.Bits				 = It.first.Bits;
+			Entry.NumHits			 = It.second;
+			Entry.FileName			 = ScavengeDatabase.GetPartialSourceFilePath(It.first);
+			Entry.FullSourceFilePath = ScavengeDatabase.GetFullSourceFilePath(It.first);
+
+			SortedPossibleSources.push_back(Entry);
+		}
+	};
+
+	std::sort(SortedPossibleSources.begin(), SortedPossibleSources.end(), [](const FPossibleSource& A, const FPossibleSource& B) {
+		return A.NumHits > B.NumHits;
+	});
+
+	const uint64 ScavengeSizeThreshold = uint64(double(TotalNeedSize) * 0.01);
+
+	std::vector<FNeedBlock> LocalNeedList;
+	for (const FPossibleSource& PossibleSource : SortedPossibleSources)
+	{
+		LocalNeedList.clear();
+
+		if (ScavengeNeedList.empty())
+		{
+			break;
+		}
+
+		const FFileManifest& ScavengeFileManifest = ScavengeDatabase.GetFileManifest(PossibleSource);
+
+		NativeFile LocalSourceFile = NativeFile(PossibleSource.FullSourceFilePath, EFileMode::ReadOnlyUnbuffered);
+		if (LocalSourceFile.IsValid())
+		{
+			THashMap<FGenericHash, uint64> BlockOffsetMap;
+			for (const FGenericBlock& Block : ScavengeFileManifest.Blocks)
+			{
+				BlockOffsetMap[Block.HashStrong] = Block.Offset;
+			}
+
+			for (const FNeedBlock& NeedBlock : ScavengeNeedList)
+			{
+				auto OffsetIt = BlockOffsetMap.find(NeedBlock.Hash);
+				if (OffsetIt != BlockOffsetMap.end())
+				{
+					FNeedBlock LocalNeedBlock	= NeedBlock;
+					LocalNeedBlock.SourceOffset = OffsetIt->second;
+					LocalNeedList.push_back(LocalNeedBlock);
+				}
+			}
+
+			const uint64 LocalNeedListSize = ComputeSize(LocalNeedList);
+			if (LocalNeedListSize < ScavengeSizeThreshold)
+			{
+				continue;
+			}
+
+			UNSYNC_VERBOSE(L"Scavenging data from '%s': %.3f MB",
+						   PossibleSource.FullSourceFilePath.wstring().c_str(),
+						   SizeMb(LocalNeedListSize));
+
+			std::sort(LocalNeedList.begin(), LocalNeedList.end(), FNeedBlock::FCompareBySourceOffset());
+
+			std::vector<FCopyCommandWithBlockRange> CopyCommands = OptimizeNeedListWithBlockRange(LocalNeedList, 1_MB);
+
+			auto ReadCallback =
+				[&Result, &OutScavengedBlocks, StrongHasher](FIOBuffer Buffer, uint64 SourceOffset, uint64 ReadSize, uint64 UserData) {
+					const FCopyCommandWithBlockRange& CopyCommand = *reinterpret_cast<const FCopyCommandWithBlockRange*>(UserData);
+
+					UNSYNC_ASSERT(CopyCommand.BlockRange.Size() != 0);
+
+					const FNeedBlock& FirstBlock = CopyCommand.BlockRange.BeginPtr[0];
+
+					uint64 BlockOffset	= 0;
+					bool   bBlockHashOk = true;
+					for (const FNeedBlock& Block : CopyCommand.BlockRange)
+					{
+						const uint8* BlockData	  = Buffer.GetData() + BlockOffset;
+						FGenericHash ActualHash	  = ComputeHash(BlockData, Block.Size, StrongHasher);
+						FGenericHash ExpectedHash = Block.Hash;
+
+						UNSYNC_ASSERT(Block.TargetOffset == FirstBlock.TargetOffset + BlockOffset);
+
+						if (ActualHash != ExpectedHash)
+						{
+							bBlockHashOk = false;
+							break;
+						}
+
+						BlockOffset += Block.Size;
+					}
+
+					if (CopyCommand.Size != ReadSize)
+					{
+						bBlockHashOk = false;
+					}
+
+					if (bBlockHashOk)
+					{
+						Result.Write(Buffer.GetData(), FirstBlock.TargetOffset, ReadSize);
+
+						for (const FNeedBlock& Block : CopyCommand.BlockRange)
+						{
+							OutScavengedBlocks.insert(Block.Hash.ToHash128());
+						}
+
+						AddGlobalProgress(ReadSize, EBlockListType::Source);
+					}
+				};
+
+			for (const FCopyCommandWithBlockRange& Command : CopyCommands)
+			{
+				LocalSourceFile.ReadAsync(Command.SourceOffset, Command.Size, (uint64)(&Command), ReadCallback);
+			}
+
+			LocalSourceFile.FlushAll();
+		}
+
+		auto FilterResult =
+			std::remove_if(ScavengeNeedList.begin(), ScavengeNeedList.end(), [&OutScavengedBlocks](const FNeedBlock& Block) {
+				return OutScavengedBlocks.find(Block.Hash.ToHash128()) != OutScavengedBlocks.end();
+			});
+
+		ScavengeNeedList.erase(FilterResult, ScavengeNeedList.end());
+	}
+}
+
+}  // namespace unsync
diff --git a/Engine/Source/Programs/Unsync/Private/UnsyncScavenger.h b/Engine/Source/Programs/Unsync/Private/UnsyncScavenger.h
new file mode 100644
--- /dev/null
+++ b/Engine/Source/Programs/Unsync/Private/UnsyncScavenger.h
@@ -0,0 +1,82 @@
+// Copyright Epic Games, Inc. All Rights Reserved.
+
+#pragma once
+
+#include "UnsyncCommon.h"
+#include "UnsyncCore.h"
+#include "UnsyncHashTable.h"
+#include "UnsyncUtil.h"
+
+#include <unordered_map>
+
+namespace unsync {
+
+struct FIOWriter;
+struct FFileSyncTask;
+struct FNeedBlock;
+
+struct FScavengedManifest
+{
+	bool			   bValid = false;
+	FPath			   Root;
+	FPath			   ManifestPath;
+	FDirectoryManifest Manifest;
+	std::vector<FPath> FileList;
+};
+
+struct FScavengeBlockSource
+{
+	union
+	{
+		struct
+		{
+			uint32 ManifestIndex;
+			uint32 FileIndex;
+		} Data;
+		uint64 Bits = ~0ull;
+	};
+
+	bool operator==(const FScavengeBlockSource& Other) const { return Bits == Other.Bits; }
+
+	struct FHash
+	{
+		uint64 operator()(const FScavengeBlockSource& X) const { return X.Bits; }
+	};
+};
+
+using FScavengeBlockMap = std::unordered_multimap<FHash128, FScavengeBlockSource>;
+
+class FScavengeDatabase
+{
+public:
+
+	static FScavengeDatabase* BuildFromFileSyncTasks(const FSyncDirectoryOptions& SyncOptions, TArrayView<FFileSyncTask> AllFileTasks);
+
+	const FScavengeBlockMap&  GetBlockMap() const { return BlockMap; }
+	const FPath&			  GetPartialSourceFilePath(FScavengeBlockSource SourceId) const;
+	FPath					  GetFullSourceFilePath(FScavengeBlockSource SourceId) const;
+	const FScavengedManifest& GetScavengedManifest(FScavengeBlockSource SourceId) const { return Manifests[SourceId.Data.ManifestIndex]; }
+	const FFileManifest&	  GetFileManifest(FScavengeBlockSource SourceId) const;
+	bool					  IsSourceValid(FScavengeBlockSource SourceId) const;
+
+private:
+	std::vector<FScavengedManifest> Manifests;
+
+	THashSet<FHash128> UniqueBlockHashes;
+	THashSet<FHash128> UniqueUsableBlockHashes;
+
+	// Which local files may contain blocks with this hash
+	FScavengeBlockMap BlockMap;
+
+	// Holds source locations that are detected to be inaccessible or corrupted.
+	// This can be used to skip trying a known-bad source over and over.
+	// THashSet<FScavengeBlockSource, FScavengeBlockSource::FHash> BannedBlockSources;
+};
+
+void BuildTargetFromScavengedData(FIOWriter&					 Result,
+								  const std::vector<FNeedBlock>& NeedList,
+								  const FScavengeDatabase&		 ScavengeDatabase,
+								  EStrongHashAlgorithmID		 StrongHasher,
+								  THashSet<FHash128>&			 OutScavengedBlocks);
+
+}  // namespace unsync
