unsync - Cleanup scavenged manifest paths

* Remve UNC prefixes on Windows

diff --git a/Engine/Source/Programs/Unsync/Private/UnsyncFile.cpp b/Engine/Source/Programs/Unsync/Private/UnsyncFile.cpp
--- a/Engine/Source/Programs/Unsync/Private/UnsyncFile.cpp
+++ b/Engine/Source/Programs/Unsync/Private/UnsyncFile.cpp
@@ -58,6 +58,28 @@ MakeExtendedAbsolutePath(const FPath& InAbsolutePath)
 #endif // UNSYNC_PLATFORM_WINDOWS
 }
 
+FPath
+RemoveExtendedPathPrefix(const FPath& InPath)
+{
+#if UNSYNC_PLATFORM_WINDOWS
+	const std::wstring& InPathString = InPath.native();
+	if (InPathString.starts_with(L"\\\\?\\UNC\\"))
+	{
+		return FPath(InPathString.substr(8));
+	}
+	else if (InPathString.starts_with(L"\\\\?\\"))
+	{
+		return FPath(InPathString.substr(4));
+	}
+	else
+	{
+		return InPath;
+	}
+#else // UNSYNC_PLATFORM_WINDOWS
+	return InPath;
+#endif // UNSYNC_PLATFORM_WINDOWS
+}
+
 #if UNSYNC_PLATFORM_WINDOWS
 inline uint64
 MakeU64(FILETIME Ft)
diff --git a/Engine/Source/Programs/Unsync/Private/UnsyncFile.h b/Engine/Source/Programs/Unsync/Private/UnsyncFile.h
--- a/Engine/Source/Programs/Unsync/Private/UnsyncFile.h
+++ b/Engine/Source/Programs/Unsync/Private/UnsyncFile.h
@@ -469,4 +469,8 @@ FFileAttributeCache CreateFileAttributeCache(const FPath& Root, const FSyncFilte
 // https://docs.microsoft.com/en-us/windows/win32/fileio/maximum-file-path-limitation
 FPath MakeExtendedAbsolutePath(const FPath& InAbsolutePath);
 
+// Removes `\\?\` or `\\?\UNC\` prefix from a given path.
+// Returns original path on non-Windows.
+FPath RemoveExtendedPathPrefix(const FPath& InPath);
+
 }  // namespace unsync
diff --git a/Engine/Source/Programs/Unsync/Private/UnsyncScavenger.cpp b/Engine/Source/Programs/Unsync/Private/UnsyncScavenger.cpp
--- a/Engine/Source/Programs/Unsync/Private/UnsyncScavenger.cpp
+++ b/Engine/Source/Programs/Unsync/Private/UnsyncScavenger.cpp
@@ -26,7 +26,7 @@ FScavengeDatabase::BuildFromFileSyncTasks(const FSyncDirectoryOptions& SyncOptio
 		return Result;
 	}
 
-	const FPath ExtendedTargetPath = MakeExtendedAbsolutePath(SyncOptions.Target);
+	const FPath ExtendedTargetPath = RemoveExtendedPathPrefix(SyncOptions.Target);
 
 	UNSYNC_VERBOSE(L"Scanning '%ls' for usable manifests", SyncOptions.ScavengeRoot.c_str());
 
@@ -64,8 +64,8 @@ FScavengeDatabase::BuildFromFileSyncTasks(const FSyncDirectoryOptions& SyncOptio
 			if (PathExists(ManifestPath))
 			{
 				FScavengedManifest Entry;
-				Entry.ManifestPath = ManifestPath;
-				Entry.Root		   = DirPath.parent_path();
+				Entry.ManifestPath = RemoveExtendedPathPrefix(ManifestPath);
+				Entry.Root		   = RemoveExtendedPathPrefix(DirPath.parent_path());
 				Result->Manifests.push_back(std::move(Entry));
 			}
 		}
@@ -224,8 +224,8 @@ struct FCopyCommandWithBlockRange : FCopyCommand
 
 // Similar to OptimizeNeedList, but assumes maintains commands in the same order as input blocks.
 // Input blocks are preferred to be sorted by source offset, but it is not a hard requirement.
-std::vector<FCopyCommandWithBlockRange> static OptimizeNeedListWithBlockRange(const std::vector<FNeedBlock>& Input,
-																			  uint64						 MaxMergedBlockSize)
+static std::vector<FCopyCommandWithBlockRange>
+OptimizeNeedListWithBlockRange(const std::vector<FNeedBlock>& Input, uint64 MaxMergedBlockSize)
 {
 	std::vector<FCopyCommandWithBlockRange> Result;
 	Result.reserve(Input.size());
@@ -297,9 +297,9 @@ BuildTargetFromScavengedData(FIOWriter&						Result,
 	uint64					 TotalNeedSize	  = 0;
 	for (const FNeedBlock& SourceNeedBlock : NeedList)
 	{
-		FHash128   NeedBlockHash = SourceNeedBlock.Hash.ToHash128();
+		FHash128 NeedBlockHash = SourceNeedBlock.Hash.ToHash128();
 		TotalNeedSize += SourceNeedBlock.Size;
-		const auto Sources		 = ScavengeBlockMap.equal_range(NeedBlockHash);
+		const auto Sources = ScavengeBlockMap.equal_range(NeedBlockHash);
 		if (Sources.first != Sources.second)
 		{
 			ScavengeNeedList.push_back(SourceNeedBlock);
