unsync - Deduplicate blocks within the same file while building scavenge database

* Log progress when scavenging blocks from local files
* Move prorress logging helper code to dedicated files

diff --git a/Engine/Source/Programs/Unsync/CMakeLists.txt b/Engine/Source/Programs/Unsync/CMakeLists.txt
--- a/Engine/Source/Programs/Unsync/CMakeLists.txt
+++ b/Engine/Source/Programs/Unsync/CMakeLists.txt
@@ -68,6 +68,8 @@ set(src
 	Private/UnsyncMiniCb.cpp
 	Private/UnsyncMiniCb.h
 	Private/UnsyncPool.h
+	Private/UnsyncProgress.cpp
+	Private/UnsyncProgress.h
 	Private/UnsyncProtocol.cpp
 	Private/UnsyncProtocol.h
 	Private/UnsyncProxy.cpp
diff --git a/Engine/Source/Programs/Unsync/Private/UnsyncCore.cpp b/Engine/Source/Programs/Unsync/Private/UnsyncCore.cpp
--- a/Engine/Source/Programs/Unsync/Private/UnsyncCore.cpp
+++ b/Engine/Source/Programs/Unsync/Private/UnsyncCore.cpp
@@ -12,6 +12,7 @@
 #include "UnsyncSerialization.h"
 #include "UnsyncThread.h"
 #include "UnsyncUtil.h"
+#include "UnsyncProgress.h"
 
 #include <condition_variable>
 #include <deque>
@@ -26,7 +27,7 @@ UNSYNC_THIRD_PARTY_INCLUDES_START
 #include <md5-sse2.h>
 UNSYNC_THIRD_PARTY_INCLUDES_END
 
-#define UNSYNC_VERSION_STR "1.0.51-dev3"
+#define UNSYNC_VERSION_STR "1.0.51-dev4"
 
 namespace unsync {
 
@@ -46,102 +47,6 @@ ComputeMaxVariableBlockSize(uint32 BlockSize)
 	return std::max(BlockSize, 4096u) * 4;	// changing this invalidates cached blocks
 }
 
-// pretend that reading remote files is ~25x slower than local
-const uint64 GLOBAL_PROGRESS_SOURCE_SCALE = 25;
-const uint64 GLOBAL_PROGRESS_BASE_SCALE	  = 1;
-
-std::atomic<uint64> GGlobalProgressCurrent;
-std::atomic<uint64> GGlobalProgressTotal;
-
-inline void
-LogGlobalProgress()
-{
-	LogProgress(nullptr, GGlobalProgressCurrent, GGlobalProgressTotal);
-}
-
-enum class ELogProgressUnits {
-	Raw,
-	Bytes,
-	MB,
-	GB,
-};
-
-struct FLogProgressScope
-{
-	UNSYNC_DISALLOW_COPY_ASSIGN(FLogProgressScope)
-
-	FLogProgressScope(uint64 InTotal, ELogProgressUnits InUnits = ELogProgressUnits::Raw, uint64 InPeriodMilliseconds = 500)
-	: Current(0)
-	, Total(InTotal)
-	, PeriodMilliseconds(InPeriodMilliseconds)
-	, Units(InUnits)
-	, NextProgressLogTime(TimePointNow())
-	, bEnabled(true)
-	{
-		Add(0);
-	}
-
-	void Complete() { Add(0, true); }
-
-	void Add(uint64 X, bool bForceComplete = false)
-	{
-		if (!bEnabled)
-		{
-			return;
-		}
-
-		Current += X;
-
-		std::lock_guard<std::mutex> LockGuard(Mutex);
-		const uint64				CurrentClamped = std::min<uint64>(Current, Total);
-		const bool					bComplete	   = (CurrentClamped == Total) || bForceComplete;
-		if (bEnabled && GLogVerbose && (TimePointNow() > NextProgressLogTime || bComplete))
-		{
-			const wchar_t* Ending = bComplete ? L"\n" : L"\r";
-			switch (Units)
-			{
-				case ELogProgressUnits::GB:
-					LogPrintf(ELogLevel::Debug,
-							  L"%.2f / %.2f GB%ls",
-							  double(CurrentClamped) / double(1_GB),
-							  double(Total) / double(1_GB),
-							  Ending);
-					break;
-				case ELogProgressUnits::MB:
-					LogPrintf(ELogLevel::Debug,
-							  L"%.2f / %.2f MB%ls",
-							  double(CurrentClamped) / double(1_MB),
-							  double(Total) / double(1_MB),
-							  Ending);
-					break;
-				case ELogProgressUnits::Bytes:
-					LogPrintf(ELogLevel::Debug, L"%llu / %llu bytes%ls", (uint64)CurrentClamped, Total, Ending);
-					break;
-				case ELogProgressUnits::Raw:
-				default:
-					LogPrintf(ELogLevel::Debug, L"%llu / %llu%ls", (uint64)CurrentClamped, Total, Ending);
-					break;
-			}
-
-			NextProgressLogTime = TimePointNow() + std::chrono::milliseconds(PeriodMilliseconds);
-			LogGlobalProgress();
-
-			if (bComplete)
-			{
-				bEnabled = false;
-			}
-		}
-	}
-
-	std::mutex				Mutex;
-	std::atomic<uint64>		Current;
-	const uint64			Total;
-	const uint64			PeriodMilliseconds;
-	const ELogProgressUnits Units;
-	FTimePoint				NextProgressLogTime;
-	std::atomic<bool>		bEnabled;
-};
-
 static uint64
 BlockingReadLarge(FIOReader& Reader, uint64 Offset, uint64 Size, uint8* OutputBuffer, uint64 OutputBufferSize)
 {
@@ -1166,21 +1071,6 @@ BuildTargetBuffer(FIOReader& SourceProvider, FIOReader& BaseProvider, const FNee
 	return Result;
 }
 
-void
-AddGlobalProgress(uint64 Size, EBlockListType ListType)
-{
-	if (ListType == EBlockListType::Source)
-	{
-		Size *= GLOBAL_PROGRESS_SOURCE_SCALE;
-	}
-	else
-	{
-		Size *= GLOBAL_PROGRESS_BASE_SCALE;
-	}
-
-	GGlobalProgressCurrent += Size;
-}
-
 void
 DownloadBlocks(FProxyPool&					  ProxyPool,
 			   const std::vector<FNeedBlock>& OriginalUniqueNeedBlocks,
diff --git a/Engine/Source/Programs/Unsync/Private/UnsyncCore.h b/Engine/Source/Programs/Unsync/Private/UnsyncCore.h
--- a/Engine/Source/Programs/Unsync/Private/UnsyncCore.h
+++ b/Engine/Source/Programs/Unsync/Private/UnsyncCore.h
@@ -379,13 +379,6 @@ enum class ESyncSourceType
 	Server,
 };
 
-enum EBlockListType {
-	Source,
-	Base
-};
-
-void AddGlobalProgress(uint64 Size, EBlockListType ListType);
-
 struct FSyncDirectoryOptions
 {
 	ESyncSourceType	   SourceType;
diff --git a/Engine/Source/Programs/Unsync/Private/UnsyncProgress.cpp b/Engine/Source/Programs/Unsync/Private/UnsyncProgress.cpp
new file mode 100644
--- /dev/null
+++ b/Engine/Source/Programs/Unsync/Private/UnsyncProgress.cpp
@@ -0,0 +1,91 @@
+#include "UnsyncProgress.h"
+
+namespace unsync {
+
+std::atomic<uint64> GGlobalProgressCurrent;
+std::atomic<uint64> GGlobalProgressTotal;
+
+void
+AddGlobalProgress(uint64 Size, EBlockListType ListType)
+{
+	if (ListType == EBlockListType::Source)
+	{
+		Size *= GLOBAL_PROGRESS_SOURCE_SCALE;
+	}
+	else
+	{
+		Size *= GLOBAL_PROGRESS_BASE_SCALE;
+	}
+
+	GGlobalProgressCurrent += Size;
+}
+
+FLogProgressScope::FLogProgressScope(uint64 InTotal, ELogProgressUnits InUnits, uint64 InPeriodMilliseconds)
+: Current(0)
+, Total(InTotal)
+, PeriodMilliseconds(InPeriodMilliseconds)
+, Units(InUnits)
+, NextProgressLogTime(TimePointNow())
+, bEnabled(true)
+{
+	Add(0);
+}
+
+void
+FLogProgressScope::Complete()
+{
+	Add(0, true);
+}
+
+void
+FLogProgressScope::Add(uint64 X, bool bForceComplete)
+{
+	if (!bEnabled)
+	{
+		return;
+	}
+
+	Current += X;
+
+	std::lock_guard<std::mutex> LockGuard(Mutex);
+	const uint64				CurrentClamped = std::min<uint64>(Current, Total);
+	const bool					bComplete	   = (CurrentClamped == Total) || bForceComplete;
+	if (bEnabled && GLogVerbose && (TimePointNow() > NextProgressLogTime || bComplete))
+	{
+		const wchar_t* Ending = bComplete ? L"\n" : L"\r";
+		switch (Units)
+		{
+			case ELogProgressUnits::GB:
+				LogPrintf(ELogLevel::Debug,
+						  L"%.2f / %.2f GB%ls",
+						  double(CurrentClamped) / double(1_GB),
+						  double(Total) / double(1_GB),
+						  Ending);
+				break;
+			case ELogProgressUnits::MB:
+				LogPrintf(ELogLevel::Debug,
+						  L"%.2f / %.2f MB%ls",
+						  double(CurrentClamped) / double(1_MB),
+						  double(Total) / double(1_MB),
+						  Ending);
+				break;
+			case ELogProgressUnits::Bytes:
+				LogPrintf(ELogLevel::Debug, L"%llu / %llu bytes%ls", (uint64)CurrentClamped, Total, Ending);
+				break;
+			case ELogProgressUnits::Raw:
+			default:
+				LogPrintf(ELogLevel::Debug, L"%llu / %llu%ls", (uint64)CurrentClamped, Total, Ending);
+				break;
+		}
+
+		NextProgressLogTime = TimePointNow() + std::chrono::milliseconds(PeriodMilliseconds);
+		LogGlobalProgress();
+
+		if (bComplete)
+		{
+			bEnabled = false;
+		}
+	}
+}
+
+}  // namespace unsync
diff --git a/Engine/Source/Programs/Unsync/Private/UnsyncProgress.h b/Engine/Source/Programs/Unsync/Private/UnsyncProgress.h
new file mode 100644
--- /dev/null
+++ b/Engine/Source/Programs/Unsync/Private/UnsyncProgress.h
@@ -0,0 +1,58 @@
+#pragma once
+
+#include "UnsyncCommon.h"
+#include "UnsyncLog.h"
+#include "UnsyncUtil.h"
+
+#include <atomic>
+#include <mutex>
+
+namespace unsync {
+
+// pretend that reading remote files is ~25x slower than local
+static constexpr uint64 GLOBAL_PROGRESS_SOURCE_SCALE = 25;
+static constexpr uint64 GLOBAL_PROGRESS_BASE_SCALE	 = 1;
+
+extern std::atomic<uint64> GGlobalProgressCurrent;
+extern std::atomic<uint64> GGlobalProgressTotal;
+
+inline void
+LogGlobalProgress()
+{
+	LogProgress(nullptr, GGlobalProgressCurrent, GGlobalProgressTotal);
+}
+
+enum class EBlockListType {
+	Source,
+	Base
+};
+
+void AddGlobalProgress(uint64 Size, EBlockListType ListType);
+
+enum class ELogProgressUnits {
+	Raw,
+	Bytes,
+	MB,
+	GB,
+};
+
+struct FLogProgressScope
+{
+	UNSYNC_DISALLOW_COPY_ASSIGN(FLogProgressScope)
+
+	FLogProgressScope(uint64 InTotal, ELogProgressUnits InUnits = ELogProgressUnits::Raw, uint64 InPeriodMilliseconds = 500);
+
+	void Complete();
+
+	void Add(uint64 X, bool bForceComplete = false);
+
+	std::mutex				Mutex;
+	std::atomic<uint64>		Current;
+	const uint64			Total;
+	const uint64			PeriodMilliseconds;
+	const ELogProgressUnits Units;
+	FTimePoint				NextProgressLogTime;
+	std::atomic<bool>		bEnabled;
+};
+
+}  // namespace unsync
diff --git a/Engine/Source/Programs/Unsync/Private/UnsyncScavenger.cpp b/Engine/Source/Programs/Unsync/Private/UnsyncScavenger.cpp
--- a/Engine/Source/Programs/Unsync/Private/UnsyncScavenger.cpp
+++ b/Engine/Source/Programs/Unsync/Private/UnsyncScavenger.cpp
@@ -3,8 +3,10 @@
 #include "UnsyncScavenger.h"
 #include "UnsyncFile.h"
 #include "UnsyncLog.h"
+#include "UnsyncProgress.h"
 #include "UnsyncSerialization.h"
 #include "UnsyncThread.h"
+
 namespace unsync {
 
 FScavengeDatabase*
@@ -87,6 +89,8 @@ FScavengeDatabase::BuildFromFileSyncTasks(const FSyncDirectoryOptions& SyncOptio
 		}
 	});
 
+	UNSYNC_VERBOSE(L"Building block database");
+
 	THashMap<FScavengeBlockSource, uint64, FScavengeBlockSource::FHash> BlockSourceUseCounts;
 
 	uint32 ManifestIndex = 0;
@@ -96,6 +100,9 @@ FScavengeDatabase::BuildFromFileSyncTasks(const FSyncDirectoryOptions& SyncOptio
 		for (const auto& FileManifestIt : Entry.Manifest.Files)
 		{
 			const FFileManifest& FileManifest = FileManifestIt.second;
+
+			THashSet<FHash128> UniqueBlocksPerFile;
+
 			for (const FGenericBlock& BlockInfo : FileManifest.Blocks)
 			{
 				FHash128 BlockHash = BlockInfo.HashStrong.ToHash128();
@@ -105,10 +112,13 @@ FScavengeDatabase::BuildFromFileSyncTasks(const FSyncDirectoryOptions& SyncOptio
 					FScavengeBlockSource BlockSource;
 					BlockSource.Data.ManifestIndex = ManifestIndex;
 					BlockSource.Data.FileIndex	   = FileIndex;
-					Result->BlockMap.insert(std::make_pair(BlockHash, BlockSource));
-					Result->UniqueUsableBlockHashes.insert(BlockHash);
 
-					BlockSourceUseCounts[BlockSource] += 1;
+					if (UniqueBlocksPerFile.insert(BlockHash).second)
+					{
+						Result->BlockMap.insert(std::make_pair(BlockHash, BlockSource));
+						Result->UniqueUsableBlockHashes.insert(BlockHash);
+						BlockSourceUseCounts[BlockSource] += 1;
+					}
 				}
 			}
 			++FileIndex;
@@ -294,11 +304,11 @@ BuildTargetFromScavengedData(FIOWriter&						Result,
 	THashMap<FScavengeBlockSource, uint64, FScavengeBlockSource::FHash> PossibleSources;
 
 	const FScavengeBlockMap& ScavengeBlockMap = ScavengeDatabase.GetBlockMap();
-	uint64					 TotalNeedSize	  = 0;
+	uint64					 TotalCopySize	  = 0;
 	for (const FNeedBlock& SourceNeedBlock : NeedList)
 	{
 		FHash128 NeedBlockHash = SourceNeedBlock.Hash.ToHash128();
-		TotalNeedSize += SourceNeedBlock.Size;
+		TotalCopySize += SourceNeedBlock.Size;
 		const auto Sources = ScavengeBlockMap.equal_range(NeedBlockHash);
 		if (Sources.first != Sources.second)
 		{
@@ -338,7 +348,9 @@ BuildTargetFromScavengedData(FIOWriter&						Result,
 		return A.NumHits > B.NumHits;
 	});
 
-	const uint64 ScavengeSizeThreshold = uint64(double(TotalNeedSize) * 0.01);
+	FLogProgressScope ProgressLogger(TotalCopySize, ELogProgressUnits::MB);
+
+	const uint64 ScavengeSizeThreshold = uint64(double(TotalCopySize) * 0.01);
 
 	std::vector<FNeedBlock> LocalNeedList;
 	for (const FPossibleSource& PossibleSource : SortedPossibleSources)
@@ -378,58 +390,61 @@ BuildTargetFromScavengedData(FIOWriter&						Result,
 				continue;
 			}
 
-			UNSYNC_VERBOSE(L"Scavenging data from '%s': %.3f MB",
-						   PossibleSource.FullSourceFilePath.wstring().c_str(),
-						   SizeMb(LocalNeedListSize));
+			UNSYNC_VERBOSE(L"Scavenging data from '%s'",
+						   PossibleSource.FullSourceFilePath.wstring().c_str());
 
 			std::sort(LocalNeedList.begin(), LocalNeedList.end(), FNeedBlock::FCompareBySourceOffset());
 
 			std::vector<FCopyCommandWithBlockRange> CopyCommands = OptimizeNeedListWithBlockRange(LocalNeedList, 1_MB);
 
-			auto ReadCallback =
-				[&Result, &OutScavengedBlocks, StrongHasher](FIOBuffer Buffer, uint64 SourceOffset, uint64 ReadSize, uint64 UserData) {
-					const FCopyCommandWithBlockRange& CopyCommand = *reinterpret_cast<const FCopyCommandWithBlockRange*>(UserData);
-
-					UNSYNC_ASSERT(CopyCommand.BlockRange.Size() != 0);
-
-					const FNeedBlock& FirstBlock = CopyCommand.BlockRange.BeginPtr[0];
+			auto ReadCallback = [&Result, &OutScavengedBlocks, &ProgressLogger, StrongHasher](FIOBuffer Buffer,
+																							  uint64	SourceOffset,
+																							  uint64	ReadSize,
+																							  uint64	UserData) {
+				const FCopyCommandWithBlockRange& CopyCommand = *reinterpret_cast<const FCopyCommandWithBlockRange*>(UserData);
 
-					uint64 BlockOffset	= 0;
-					bool   bBlockHashOk = true;
-					for (const FNeedBlock& Block : CopyCommand.BlockRange)
-					{
-						const uint8* BlockData	  = Buffer.GetData() + BlockOffset;
-						FGenericHash ActualHash	  = ComputeHash(BlockData, Block.Size, StrongHasher);
-						FGenericHash ExpectedHash = Block.Hash;
+				UNSYNC_ASSERT(CopyCommand.BlockRange.Size() != 0);
 
-						UNSYNC_ASSERT(Block.TargetOffset == FirstBlock.TargetOffset + BlockOffset);
+				const FNeedBlock& FirstBlock = CopyCommand.BlockRange.BeginPtr[0];
 
-						if (ActualHash != ExpectedHash)
-						{
-							bBlockHashOk = false;
-							break;
-						}
+				uint64 BlockOffset	= 0;
+				bool   bBlockHashOk = true;
+				for (const FNeedBlock& Block : CopyCommand.BlockRange)
+				{
+					const uint8* BlockData	  = Buffer.GetData() + BlockOffset;
+					FGenericHash ActualHash	  = ComputeHash(BlockData, Block.Size, StrongHasher);
+					FGenericHash ExpectedHash = Block.Hash;
 
-						BlockOffset += Block.Size;
-					}
+					UNSYNC_ASSERT(Block.TargetOffset == FirstBlock.TargetOffset + BlockOffset);
 
-					if (CopyCommand.Size != ReadSize)
+					if (ActualHash != ExpectedHash)
 					{
 						bBlockHashOk = false;
+						break;
 					}
 
-					if (bBlockHashOk)
-					{
-						Result.Write(Buffer.GetData(), FirstBlock.TargetOffset, ReadSize);
+					BlockOffset += Block.Size;
+				}
+
+				if (CopyCommand.Size != ReadSize)
+				{
+					bBlockHashOk = false;
+				}
 
-						for (const FNeedBlock& Block : CopyCommand.BlockRange)
-						{
-							OutScavengedBlocks.insert(Block.Hash.ToHash128());
-						}
+				if (bBlockHashOk)
+				{
+					Result.Write(Buffer.GetData(), FirstBlock.TargetOffset, ReadSize);
 
-						AddGlobalProgress(ReadSize, EBlockListType::Source);
+					for (const FNeedBlock& Block : CopyCommand.BlockRange)
+					{
+						OutScavengedBlocks.insert(Block.Hash.ToHash128());
 					}
-				};
+
+					AddGlobalProgress(ReadSize, EBlockListType::Source);
+
+					ProgressLogger.Add(ReadSize);
+				}
+			};
 
 			for (const FCopyCommandWithBlockRange& Command : CopyCommands)
 			{
