unsync - Fix missing copyright boilerplate

diff --git a/Engine/Source/Programs/Unsync/Private/UnsyncProgress.cpp b/Engine/Source/Programs/Unsync/Private/UnsyncProgress.cpp
--- a/Engine/Source/Programs/Unsync/Private/UnsyncProgress.cpp
+++ b/Engine/Source/Programs/Unsync/Private/UnsyncProgress.cpp
@@ -1,3 +1,5 @@
+// Copyright Epic Games, Inc. All Rights Reserved.
+
 #include "UnsyncProgress.h"
 
 namespace unsync {
diff --git a/Engine/Source/Programs/Unsync/Private/UnsyncProgress.h b/Engine/Source/Programs/Unsync/Private/UnsyncProgress.h
--- a/Engine/Source/Programs/Unsync/Private/UnsyncProgress.h
+++ b/Engine/Source/Programs/Unsync/Private/UnsyncProgress.h
@@ -1,3 +1,5 @@
+// Copyright Epic Games, Inc. All Rights Reserved.
+
 #pragma once
 
 #include "UnsyncCommon.h"
