unsync - Finish up --scavenge mode implementation, bump version to 1.0.51

* Count scavenged blocks as "base" data for stats/telemetry
* Gracefully handle missing/invalid scavenge root path
* Cancel reading blocks from a particular source if a hash mismatch was found

diff --git a/Engine/Source/Programs/Unsync/Private/UnsyncCore.cpp b/Engine/Source/Programs/Unsync/Private/UnsyncCore.cpp
--- a/Engine/Source/Programs/Unsync/Private/UnsyncCore.cpp
+++ b/Engine/Source/Programs/Unsync/Private/UnsyncCore.cpp
@@ -27,7 +27,7 @@ UNSYNC_THIRD_PARTY_INCLUDES_START
 #include <md5-sse2.h>
 UNSYNC_THIRD_PARTY_INCLUDES_END
 
-#define UNSYNC_VERSION_STR "1.0.51-dev4"
+#define UNSYNC_VERSION_STR "1.0.51"
 
 namespace unsync {
 
@@ -1215,8 +1215,8 @@ DownloadBlocks(FProxyPool&					  ProxyPool,
 	}
 }
 
-void
-BuildTarget(FIOWriter&			   Result,
+FBuildTargetResult
+BuildTarget(FIOWriter&			   Output,
 			FIOReader&			   Source,
 			FIOReader&			   Base,
 			const FNeedList&	   NeedList,
@@ -1227,10 +1227,12 @@ BuildTarget(FIOWriter&			   Result,
 {
 	UNSYNC_LOG_INDENT;
 
+	FBuildTargetResult BuildResult;
+
 	auto TimeBegin = TimePointNow();
 
 	const FNeedListSize SizeInfo = ComputeNeedListSize(NeedList);
-	UNSYNC_ASSERT(SizeInfo.TotalBytes == Result.GetSize());
+	UNSYNC_ASSERT(SizeInfo.TotalBytes == Output.GetSize());
 
 	std::atomic<bool> bGotError = false;
 
@@ -1238,6 +1240,13 @@ BuildTarget(FIOWriter&			   Result,
 	std::atomic<bool> bSourceDataCopyTaskDone = false;
 	std::atomic<bool> bWaitingForBaseData	  = false;
 
+	struct FStats
+	{
+		std::atomic<uint64> WrittenBytesFromSource = {};
+		std::atomic<uint64> WrittenBytesFromBase   = {};
+	};
+	FStats Stats;
+
 	FSemaphore WriteSemaphore(MAX_ACTIVE_READERS);	// throttle writing tasks to avoid memory bloat
 	FTaskGroup WriteTasks;
 
@@ -1257,7 +1266,11 @@ BuildTarget(FIOWriter&			   Result,
 	{
 		THashSet<FHash128> ScavengedBlocks;
 
-		BuildTargetFromScavengedData(Result, NeedList.Source, *ScavengeDatabase, StrongHasher, ScavengedBlocks);
+		FScavengedBuildTargetResult ScavengeResult =
+			BuildTargetFromScavengedData(Output, NeedList.Source, *ScavengeDatabase, StrongHasher, ScavengedBlocks);
+
+		// Treat scavenged data the same as base for stats purposes
+		Stats.WrittenBytesFromBase += ScavengeResult.ScavengedBytes;
 
 		FilteredSourceNeedList = NeedList.Source;
 		auto FilterResult =
@@ -1271,7 +1284,7 @@ BuildTarget(FIOWriter&			   Result,
 		FilteredSourceNeedList = NeedList.Source;
 	}
 
-	auto ProcessNeedList = [bAllowVerboseLog, LogIndent, &Result, &bGotError, &WriteSemaphore, &WriteTasks, &bWaitingForBaseData](
+	auto ProcessNeedList = [bAllowVerboseLog, LogIndent, &Output, &bGotError, &WriteSemaphore, &WriteTasks, &bWaitingForBaseData, &Stats](
 							   FIOReader&					  DataProvider,
 							   const std::vector<FNeedBlock>& NeedBlocks,
 							   uint64						  TotalCopySize,
@@ -1348,13 +1361,34 @@ BuildTarget(FIOWriter&			   Result,
 			UNSYNC_ASSERT(Block.SourceOffset + Block.Size <= DataProvider.GetSize());
 			uint64 ReadBytes = 0;
 
-			auto ReadCallback = [&ReadBytes, &Result, &WriteTasks, &bGotError, &WriteSemaphore, Block, ListType](FIOBuffer CmdBuffer,
-																												 uint64	   CmdOffset,
-																												 uint64	   CmdReadSize,
-																												 uint64	   CmdUserData) {
+			auto ReadCallback = [&ReadBytes, &Output, &WriteTasks, &bGotError, &WriteSemaphore, &Stats, Block, ListType](
+									FIOBuffer CmdBuffer,
+									uint64	  CmdOffset,
+									uint64	  CmdReadSize,
+									uint64	  CmdUserData) {
 				WriteSemaphore.Acquire();
-				WriteTasks.run([Buffer = MakeShared(std::move(CmdBuffer)), CmdReadSize, Block, &Result, &bGotError, &WriteSemaphore]() {
-					uint64 WrittenBytes = Result.Write(Buffer->GetData(), Block.TargetOffset, CmdReadSize);
+				WriteTasks.run([Buffer = MakeShared(std::move(CmdBuffer)),
+								CmdReadSize,
+								Block,
+								&Output,
+								&bGotError,
+								&WriteSemaphore,
+								&Stats,
+								ListType]() {
+					const uint64 WrittenBytes = Output.Write(Buffer->GetData(), Block.TargetOffset, CmdReadSize);
+
+					if (ListType == EBlockListType::Source)
+					{
+						Stats.WrittenBytesFromSource += WrittenBytes;
+					}
+					else if (ListType == EBlockListType::Base)
+					{
+						Stats.WrittenBytesFromBase += WrittenBytes;
+					}
+					else
+					{
+						UNSYNC_FATAL(L"Unexpected block list type");
+					}
 
 					WriteSemaphore.Release();
 
@@ -1393,7 +1427,7 @@ BuildTarget(FIOWriter&			   Result,
 
 		DataProvider.FlushAll();
 
-		Result.FlushAll();
+		Output.FlushAll();
 
 		ProgressLogger.Complete();
 
@@ -1426,7 +1460,10 @@ BuildTarget(FIOWriter&			   Result,
 			{
 				FBufferView BlockBuffer = It->second;
 
-				Result.Write(BlockBuffer.Data, NeedBlock.TargetOffset, BlockBuffer.Size);
+				const uint64 WrittenBytes = Output.Write(BlockBuffer.Data, NeedBlock.TargetOffset, BlockBuffer.Size);
+
+				Stats.WrittenBytesFromSource += WrittenBytes;
+
 				AddGlobalProgress(NeedBlock.Size, EBlockListType::Source);
 			}
 		}
@@ -1484,7 +1521,7 @@ BuildTarget(FIOWriter&			   Result,
 
 		// limit how many decompression tasks can be queued up to avoid memory bloat
 		const uint64 MaxConcurrentDecompressionTasks = 64;
-		FSemaphore DecompressionSemaphore(MaxConcurrentDecompressionTasks);
+		FSemaphore	 DecompressionSemaphore(MaxConcurrentDecompressionTasks);
 
 		const bool			bParentThreadVerbose = GLogVerbose;
 		const uint32		ParentThreadIndent	 = GLogIndent;
@@ -1499,7 +1536,7 @@ BuildTarget(FIOWriter&			   Result,
 			 &BlockWriteMap,
 			 &DecompressionSemaphore,
 			 &DecompressTasks,
-			 &Result,
+			 &Output,
 			 StrongHasher,
 			 &BlockScatterMap,
 			 &DownloadedBlocksMutex,
@@ -1507,7 +1544,8 @@ BuildTarget(FIOWriter&			   Result,
 			 &bGotError,
 			 &NumHashMismatches,
 			 &ParentThreadIndent,
-			 &bParentThreadVerbose](const FDownloadedBlock& Block, FHash128 BlockHash) {
+			 &bParentThreadVerbose,
+			 &Stats](const FDownloadedBlock& Block, FHash128 BlockHash) {
 				FLogIndentScope IndentScope(ParentThreadIndent, true);
 
 				{
@@ -1529,7 +1567,7 @@ BuildTarget(FIOWriter&			   Result,
 
 					DecompressionSemaphore.Acquire();
 
-					DecompressTasks.run([&Result,
+					DecompressTasks.run([&Output,
 										 BlockHashUnaligned = BlockHash,
 										 Cmd,
 										 DownloadedData	  = std::make_shared<FIOBuffer>(std::move(DownloadedData)),
@@ -1542,7 +1580,8 @@ BuildTarget(FIOWriter&			   Result,
 										 &DownloadedBlocks,
 										 &DecompressionSemaphore,
 										 &bGotError,
-										 &NumHashMismatches]() {
+										 &NumHashMismatches,
+										 &Stats]() {
 						FLogIndentScope IndentScope(ParentThreadIndent, true);
 
 						// captured objects don't respect large alignment requirements
@@ -1574,7 +1613,10 @@ BuildTarget(FIOWriter&			   Result,
 
 								if (bOk)
 								{
-									uint64 WrittenBytes = Result.Write(DecompressedData.GetData(), Cmd.TargetOffset, Cmd.Size);
+									uint64 WrittenBytes = Output.Write(DecompressedData.GetData(), Cmd.TargetOffset, Cmd.Size);
+
+									Stats.WrittenBytesFromSource += WrittenBytes;
+
 									AddGlobalProgress(Cmd.Size, EBlockListType::Source);
 
 									if (WrittenBytes != Cmd.Size)
@@ -1604,7 +1646,10 @@ BuildTarget(FIOWriter&			   Result,
 								const std::vector<uint64>& ScatterList = ScatterIt->second;
 								for (uint64 ScatterOffset : ScatterList)
 								{
-									uint64 WrittenBytes = Result.Write(DecompressedData.GetData(), ScatterOffset, Cmd.Size);
+									uint64 WrittenBytes = Output.Write(DecompressedData.GetData(), ScatterOffset, Cmd.Size);
+
+									Stats.WrittenBytesFromSource += WrittenBytes;
+
 									if (WrittenBytes != Cmd.Size)
 									{
 										UNSYNC_FATAL(L"Expected to write %llu bytes, but written %llu", Cmd.Size, WrittenBytes);
@@ -1683,6 +1728,12 @@ BuildTarget(FIOWriter&			   Result,
 	{
 		LogGlobalProgress();
 	}
+
+	BuildResult.bSuccess	= bGotError.load();
+	BuildResult.BaseBytes	= Stats.WrittenBytesFromBase;
+	BuildResult.SourceBytes = Stats.WrittenBytesFromSource;
+
+	return BuildResult;
 }
 
 bool
@@ -2372,17 +2423,17 @@ SyncFile(const FNeedList&		   NeedList,
 			return std::unique_ptr<NativeFile>(new NativeFile(SourceFilePath, EFileMode::ReadOnlyUnbuffered));
 		});
 
-		BuildTarget(*TargetFile,
-					SourceFile,
-					BaseDataReader,
-					NeedList,
-					Options.Algorithm.StrongHashAlgorithmId,
-					Options.ProxyPool,
-					Options.BlockCache,
-					Options.ScavengeDatabase);
+		FBuildTargetResult BuildResult = BuildTarget(*TargetFile,
+													 SourceFile,
+													 BaseDataReader,
+													 NeedList,
+													 Options.Algorithm.StrongHashAlgorithmId,
+													 Options.ProxyPool,
+													 Options.BlockCache,
+													 Options.ScavengeDatabase);
 
-		Result.SourceBytes = ComputeSize(NeedList.Source);
-		Result.BaseBytes   = ComputeSize(NeedList.Base);
+		Result.SourceBytes = BuildResult.SourceBytes;
+		Result.BaseBytes   = BuildResult.BaseBytes;
 
 		if (Options.bValidateTargetFiles)
 		{
@@ -2441,6 +2492,21 @@ SyncFile(const FNeedList&		   NeedList,
 				Result.SystemErrorCode = Ec;
 			}
 		}
+
+		const uint64 ExpectedSourceBytes	= ComputeSize(NeedList.Source);
+		const uint64 ExpectedBaseBytes		= ComputeSize(NeedList.Base);
+
+		const uint64 ActualProcessedBytes	= BuildResult.SourceBytes + BuildResult.BaseBytes;
+		const uint64 ExpectedProcessedBytes = ExpectedSourceBytes + ExpectedBaseBytes;
+
+		if (ActualProcessedBytes != ExpectedProcessedBytes)
+		{
+			Result.Status = EFileSyncStatus::ErrorValidation;
+			UNSYNC_ERROR(L"Failed to patch file '%ls'. Expected to write %llu bytes, but actually wrote %llu bytes.",
+						 TargetFilePath.wstring().c_str(),
+						 llu(ExpectedProcessedBytes),
+						 llu(ActualProcessedBytes));
+		}
 	}
 	else
 	{
@@ -4331,15 +4397,33 @@ GetManifestInfo(const FDirectoryManifest& Manifest)
 	Result.NumBlocks	  = 0;
 	Result.NumMacroBlocks = 0;
 	Result.TotalSize	  = 0;
+	Result.UniqueSize	  = 0;
+
+	THashSet<FGenericHash> UniqueBlockSet;
+	THashSet<FGenericHash> UniqueMacroBlockSet;
 
 	for (const auto& It : Manifest.Files)
 	{
 		const FFileManifest& File = It.second;
-		Result.NumBlocks += File.Blocks.size();
-		Result.NumMacroBlocks += File.MacroBlocks.size();
+
+		for (const FGenericBlock& Block : File.Blocks)
+		{
+			if (UniqueBlockSet.insert(Block.HashStrong).second)
+			{
+				Result.UniqueSize += Block.Size;
+			}
+		}
+
+		for (const FGenericBlock& Block : File.MacroBlocks)
+		{
+			UniqueMacroBlockSet.insert(Block.HashStrong);
+		}
+
 		Result.TotalSize += File.Size;
 	}
 
+	Result.NumBlocks	  = UniqueBlockSet.size();
+	Result.NumMacroBlocks = UniqueMacroBlockSet.size();
 	Result.NumFiles		  = Manifest.Files.size();
 	Result.Algorithm	  = Manifest.Options;
 	Result.SerializedHash = ComputeSerializedManifestHash(Manifest);
@@ -4361,6 +4445,7 @@ LogManifestInfo(ELogLevel LogLevel, const FDirectoryManifestInfo& Info)
 	LogPrintf(LogLevel, L"Files: %llu\n", llu(Info.NumFiles));
 	LogPrintf(LogLevel, L"Blocks: %llu\n", llu(Info.NumBlocks));
 	LogPrintf(LogLevel, L"Macro blocks: %llu\n", llu(Info.NumMacroBlocks));
+	LogPrintf(LogLevel, L"Unique data size: %.0f MB (%llu bytes)\n", SizeMb(Info.UniqueSize), llu(Info.UniqueSize));
 	LogPrintf(LogLevel, L"Total data size: %.0f MB (%llu bytes)\n", SizeMb(Info.TotalSize), llu(Info.TotalSize));
 
 	// TODO: block size distribution histogram
diff --git a/Engine/Source/Programs/Unsync/Private/UnsyncCore.h b/Engine/Source/Programs/Unsync/Private/UnsyncCore.h
--- a/Engine/Source/Programs/Unsync/Private/UnsyncCore.h
+++ b/Engine/Source/Programs/Unsync/Private/UnsyncCore.h
@@ -180,6 +180,7 @@ struct FDirectoryManifest
 struct FDirectoryManifestInfo
 {
 	uint64	 TotalSize		= 0;
+	uint64	 UniqueSize		= 0;
 	uint64	 NumBlocks		= 0;
 	uint64	 NumMacroBlocks = 0;
 	uint64	 NumFiles		= 0;
@@ -250,7 +251,14 @@ FNeedList DiffManifestBlocks(const FGenericBlockArray& SourceBlocks, const FGene
 
 std::vector<FCopyCommand> OptimizeNeedList(const std::vector<FNeedBlock>& Input, uint64 MaxMergedBlockSize = 8_MB);
 
-void BuildTarget(FIOWriter&				Result,
+struct FBuildTargetResult
+{
+	bool   bSuccess	   = false;
+	uint64 SourceBytes = 0;
+	uint64 BaseBytes   = 0;
+};
+
+FBuildTargetResult BuildTarget(FIOWriter&			  Result,
 				 FIOReader&				Source,
 				 FIOReader&				Base,
 				 const FNeedList&		NeedList,
diff --git a/Engine/Source/Programs/Unsync/Private/UnsyncMain.cpp b/Engine/Source/Programs/Unsync/Private/UnsyncMain.cpp
--- a/Engine/Source/Programs/Unsync/Private/UnsyncMain.cpp
+++ b/Engine/Source/Programs/Unsync/Private/UnsyncMain.cpp
@@ -648,6 +648,12 @@ InnerMain(int Argc, char** Argv)
 	}
 	else if (Cli.got_subcommand(SubSync))
 	{
+		if (!ScavengeRoot.empty() && !IsDirectory(ScavengeRoot))
+		{
+			UNSYNC_WARNING(L"Scavenge directory '%ls' does not exist", ScavengeRoot.wstring().c_str());
+			ScavengeRoot = FPath{};
+		}
+
 		FCmdSyncOptions SyncOptions;
 
 		SyncOptions.Algorithm			   = Algorithm;
diff --git a/Engine/Source/Programs/Unsync/Private/UnsyncScavenger.cpp b/Engine/Source/Programs/Unsync/Private/UnsyncScavenger.cpp
--- a/Engine/Source/Programs/Unsync/Private/UnsyncScavenger.cpp
+++ b/Engine/Source/Programs/Unsync/Private/UnsyncScavenger.cpp
@@ -292,13 +292,15 @@ OptimizeNeedListWithBlockRange(const std::vector<FNeedBlock>& Input, uint64 MaxM
 	return Result;
 }
 
-void
-BuildTargetFromScavengedData(FIOWriter&						Result,
+FScavengedBuildTargetResult
+BuildTargetFromScavengedData(FIOWriter&						Output,
 							 const std::vector<FNeedBlock>& NeedList,
 							 const FScavengeDatabase&		ScavengeDatabase,
 							 EStrongHashAlgorithmID			StrongHasher,
 							 THashSet<FHash128>&			OutScavengedBlocks)
 {
+	FScavengedBuildTargetResult BuildResult;
+
 	std::vector<FNeedBlock> ScavengeNeedList;
 
 	THashMap<FScavengeBlockSource, uint64, FScavengeBlockSource::FHash> PossibleSources;
@@ -390,17 +392,19 @@ BuildTargetFromScavengedData(FIOWriter&						Result,
 				continue;
 			}
 
-			UNSYNC_VERBOSE(L"Scavenging data from '%s'",
-						   PossibleSource.FullSourceFilePath.wstring().c_str());
+			UNSYNC_VERBOSE(L"Scavenging data from '%s'", PossibleSource.FullSourceFilePath.wstring().c_str());
 
 			std::sort(LocalNeedList.begin(), LocalNeedList.end(), FNeedBlock::FCompareBySourceOffset());
 
 			std::vector<FCopyCommandWithBlockRange> CopyCommands = OptimizeNeedListWithBlockRange(LocalNeedList, 1_MB);
 
-			auto ReadCallback = [&Result, &OutScavengedBlocks, &ProgressLogger, StrongHasher](FIOBuffer Buffer,
-																							  uint64	SourceOffset,
-																							  uint64	ReadSize,
-																							  uint64	UserData) {
+			bool bFoundInvalidBlock = false;
+
+			auto ReadCallback = [&Output, &OutScavengedBlocks, &ProgressLogger, &BuildResult, &bFoundInvalidBlock, StrongHasher](
+									FIOBuffer Buffer,
+									uint64	  SourceOffset,
+									uint64	  ReadSize,
+									uint64	  UserData) {
 				const FCopyCommandWithBlockRange& CopyCommand = *reinterpret_cast<const FCopyCommandWithBlockRange*>(UserData);
 
 				UNSYNC_ASSERT(CopyCommand.BlockRange.Size() != 0);
@@ -433,22 +437,33 @@ BuildTargetFromScavengedData(FIOWriter&						Result,
 
 				if (bBlockHashOk)
 				{
-					Result.Write(Buffer.GetData(), FirstBlock.TargetOffset, ReadSize);
+					Output.Write(Buffer.GetData(), FirstBlock.TargetOffset, ReadSize);
 
 					for (const FNeedBlock& Block : CopyCommand.BlockRange)
 					{
 						OutScavengedBlocks.insert(Block.Hash.ToHash128());
 					}
 
+					BuildResult.ScavengedBytes += ReadSize;
 					AddGlobalProgress(ReadSize, EBlockListType::Source);
-
 					ProgressLogger.Add(ReadSize);
 				}
+				else
+				{
+					bFoundInvalidBlock = true;
+				}
+				
 			};
 
 			for (const FCopyCommandWithBlockRange& Command : CopyCommands)
 			{
 				LocalSourceFile.ReadAsync(Command.SourceOffset, Command.Size, (uint64)(&Command), ReadCallback);
+				if (bFoundInvalidBlock)
+				{
+					// If an invalid block is encountered in a source file, then assume that is corrupt and should not be used from this point.
+					// TODO: could also mark the source file as "bad" globally, to avoid using it when patching other files.
+					break;
+				}
 			}
 
 			LocalSourceFile.FlushAll();
@@ -461,6 +476,8 @@ BuildTargetFromScavengedData(FIOWriter&						Result,
 
 		ScavengeNeedList.erase(FilterResult, ScavengeNeedList.end());
 	}
+
+	return BuildResult;
 }
 
 }  // namespace unsync
diff --git a/Engine/Source/Programs/Unsync/Private/UnsyncScavenger.h b/Engine/Source/Programs/Unsync/Private/UnsyncScavenger.h
--- a/Engine/Source/Programs/Unsync/Private/UnsyncScavenger.h
+++ b/Engine/Source/Programs/Unsync/Private/UnsyncScavenger.h
@@ -73,10 +73,15 @@ private:
 	// THashSet<FScavengeBlockSource, FScavengeBlockSource::FHash> BannedBlockSources;
 };
 
-void BuildTargetFromScavengedData(FIOWriter&					 Result,
-								  const std::vector<FNeedBlock>& NeedList,
-								  const FScavengeDatabase&		 ScavengeDatabase,
-								  EStrongHashAlgorithmID		 StrongHasher,
-								  THashSet<FHash128>&			 OutScavengedBlocks);
+struct FScavengedBuildTargetResult
+{
+	uint64 ScavengedBytes = 0;
+};
+
+FScavengedBuildTargetResult BuildTargetFromScavengedData(FIOWriter&						Output,
+														 const std::vector<FNeedBlock>& NeedList,
+														 const FScavengeDatabase&		ScavengeDatabase,
+														 EStrongHashAlgorithmID			StrongHasher,
+														 THashSet<FHash128>&			OutScavengedBlocks);
 
 }  // namespace unsync
