unsync - Improved manifest info command

* Allow --include and --exclude options
* Compute approximate patch size from manifest A to B

diff --git a/Engine/Source/Programs/Unsync/Private/UnsyncCore.cpp b/Engine/Source/Programs/Unsync/Private/UnsyncCore.cpp
--- a/Engine/Source/Programs/Unsync/Private/UnsyncCore.cpp
+++ b/Engine/Source/Programs/Unsync/Private/UnsyncCore.cpp
@@ -27,7 +27,7 @@ UNSYNC_THIRD_PARTY_INCLUDES_START
 #include <md5-sse2.h>
 UNSYNC_THIRD_PARTY_INCLUDES_END
 
-#define UNSYNC_VERSION_STR "1.0.51"
+#define UNSYNC_VERSION_STR "1.0.52-dev"
 
 namespace unsync {
 
@@ -4513,8 +4513,14 @@ LogManifestDiff(ELogLevel LogLevel, const FDirectoryManifest& ManifestA, const F
 
 	uint32 NumCommonBlocks		= 0;
 	uint64 TotalCommonBlockSize = 0;
+	uint64 TotalSizeA			= 0;
+	uint64 TotalSizeB			= 0;
+
+	uint64 PatchSizeFromAtoB = 0;
+
 	for (const auto& ItA : BlocksA)
 	{
+		TotalSizeA += ItA.second.Size;
 		auto ItB = BlocksB.find(ItA.first);
 		if (ItB != BlocksB.end())
 		{
@@ -4523,6 +4529,15 @@ LogManifestDiff(ELogLevel LogLevel, const FDirectoryManifest& ManifestA, const F
 		}
 	}
 
+	for (const auto& ItB : BlocksB)
+	{
+		TotalSizeB += ItB.second.Size;
+		if (BlocksA.find(ItB.first) == BlocksA.end())
+		{
+			PatchSizeFromAtoB += ItB.second.Size;
+		}
+	}
+
 	uint32 NumCommonMacroBlocks		 = 0;
 	uint64 TotalCommonMacroBlockSize = 0;
 	for (const auto& ItA : MacroBlocksA)
@@ -4535,19 +4550,44 @@ LogManifestDiff(ELogLevel LogLevel, const FDirectoryManifest& ManifestA, const F
 		}
 	}
 
-	LogPrintf(LogLevel, L"Common blocks: %d, %.2f MB\n", NumCommonBlocks, SizeMb(TotalCommonBlockSize));
-	LogPrintf(LogLevel, L"Common macro blocks: %d, %.2f MB\n", NumCommonMacroBlocks, SizeMb(TotalCommonMacroBlockSize));
+	LogPrintf(LogLevel, L"Common macro blocks: %d, %.3f MB\n", NumCommonMacroBlocks, SizeMb(TotalCommonMacroBlockSize));
+
+	LogPrintf(LogLevel,
+			  L"Common blocks: %d, %.3f MB (%.2f%% of A, %.2f%% of B)\n",
+			  NumCommonBlocks,
+			  SizeMb(TotalCommonBlockSize),
+			  100.0 * double(TotalCommonBlockSize) / double(TotalSizeA),
+			  100.0 * double(TotalCommonBlockSize) / double(TotalSizeB));
+
+	LogPrintf(LogLevel, L"Patch size: %.3f MB\n", SizeMb(PatchSizeFromAtoB));
+}
+
+void
+FilterManifest(const FSyncFilter& SyncFilter, FDirectoryManifest& Manifest)
+{
+	auto It = Manifest.Files.begin();
+	while (It != Manifest.Files.end())
+	{
+		if (SyncFilter.ShouldSync(It->first))
+		{
+			++It;
+		}
+		else
+		{
+			It = Manifest.Files.erase(It);
+		}
+	}
 }
 
 int32
-CmdInfo(const FPath& InputA, const FPath& InputB, bool bListFiles)
+CmdInfo(const FCmdInfoOptions& Options)
 {
-	FPath DirectoryManifestPathA = IsDirectory(InputA) ? (InputA / ".unsync" / "manifest.bin") : InputA;
-	FPath DirectoryManifestPathB = IsDirectory(InputB) ? (InputB / ".unsync" / "manifest.bin") : InputB;
+	FPath DirectoryManifestPathA = IsDirectory(Options.InputA) ? (Options.InputA / ".unsync" / "manifest.bin") : Options.InputA;
+	FPath DirectoryManifestPathB = IsDirectory(Options.InputB) ? (Options.InputB / ".unsync" / "manifest.bin") : Options.InputB;
 
 	FDirectoryManifest ManifestA;
 
-	bool bManifestAValid = LoadDirectoryManifest(ManifestA, InputA, DirectoryManifestPathA);
+	bool bManifestAValid = LoadDirectoryManifest(ManifestA, Options.InputA, DirectoryManifestPathA);
 
 	if (!bManifestAValid)
 	{
@@ -4556,18 +4596,23 @@ CmdInfo(const FPath& InputA, const FPath& InputB, bool bListFiles)
 
 	LogPrintf(ELogLevel::Info, L"Manifest A: %ls\n", DirectoryManifestPathA.wstring().c_str());
 
+	if (Options.SyncFilter)
+	{
+		FilterManifest(*Options.SyncFilter, ManifestA);
+	}
+
 	{
 		UNSYNC_LOG_INDENT;
 		LogManifestInfo(ELogLevel::Info, ManifestA);
 	}
 
-	if (bListFiles)
+	if (Options.bListFiles)
 	{
 		UNSYNC_LOG_INDENT;
 		LogManifestFiles(ELogLevel::Info, ManifestA);
 	}
 
-	if (InputB.empty())
+	if (Options.InputB.empty())
 	{
 		return 0;
 	}
@@ -4576,7 +4621,7 @@ CmdInfo(const FPath& InputA, const FPath& InputB, bool bListFiles)
 
 	FDirectoryManifest ManifestB;
 
-	bool bManifestBValid = LoadDirectoryManifest(ManifestB, InputB, DirectoryManifestPathB);
+	bool bManifestBValid = LoadDirectoryManifest(ManifestB, Options.InputB, DirectoryManifestPathB);
 
 	if (!bManifestBValid)
 	{
@@ -4585,16 +4630,22 @@ CmdInfo(const FPath& InputA, const FPath& InputB, bool bListFiles)
 
 	LogPrintf(ELogLevel::Info, L"Manifest B: %ls\n", DirectoryManifestPathB.wstring().c_str());
 
+	if (Options.SyncFilter)
+	{
+		FilterManifest(*Options.SyncFilter, ManifestB);
+	}
+
 	{
 		UNSYNC_LOG_INDENT;
 		LogManifestInfo(ELogLevel::Info, ManifestB);
 	}
-	if (bListFiles)
+	if (Options.bListFiles)
 	{
 		UNSYNC_LOG_INDENT;
 		LogManifestFiles(ELogLevel::Info, ManifestB);
 	}
 
+	LogPrintf(ELogLevel::Info, L"\n");
 	LogPrintf(ELogLevel::Info, L"Difference:\n");
 
 	{
diff --git a/Engine/Source/Programs/Unsync/Private/UnsyncCore.h b/Engine/Source/Programs/Unsync/Private/UnsyncCore.h
--- a/Engine/Source/Programs/Unsync/Private/UnsyncCore.h
+++ b/Engine/Source/Programs/Unsync/Private/UnsyncCore.h
@@ -427,6 +427,13 @@ FHash160 ComputeSerializedManifestHash160(const FDirectoryManifest& Manifest);
 FBlock128			   ToBlock128(const FGenericBlock& GenericBlock);
 std::vector<FBlock128> ToBlock128(FGenericBlockArray& GenericBlocks);
 
-int32 CmdInfo(const FPath& InputA, const FPath& InputB, bool bListFiles);
+struct FCmdInfoOptions
+{
+	FPath InputA;
+	FPath InputB;
+	bool bListFiles = false;
+	const FSyncFilter* SyncFilter = nullptr;
+};
+int32 CmdInfo(const FCmdInfoOptions& Options);
 
 }  // namespace unsync
diff --git a/Engine/Source/Programs/Unsync/Private/UnsyncMain.cpp b/Engine/Source/Programs/Unsync/Private/UnsyncMain.cpp
--- a/Engine/Source/Programs/Unsync/Private/UnsyncMain.cpp
+++ b/Engine/Source/Programs/Unsync/Private/UnsyncMain.cpp
@@ -146,6 +146,13 @@ InnerMain(int Argc, char** Argv)
 	SubInfo->add_option("Input 1", InputFilenameUtf8, "Input manifest file or root directory")->required();
 	SubInfo->add_option("Input 2", InputFilename2Utf8, "Optional input manifest file or root directory");
 	SubInfo->add_flag("--files", bInfoFiles, "List all files in the manifest");
+	SubInfo->add_option(
+		"--include",
+		IncludeFilterArrayUtf8,
+		"Include filenames that contain specified words (comma separated). If this is not present, all files will be included.");
+	SubInfo->add_option("--exclude",
+						ExcludeFilterArrayUtf8,
+						"Exclude filenames that contain specified words (comma separated). Filter is run after --include.");
 	SubCommands.push_back(SubInfo);
 
 	CLI::App* SubDiff = Cli.add_subcommand("diff", "Compute difference required to transform BaseFile into SourceFile");
@@ -702,7 +709,12 @@ InnerMain(int Argc, char** Argv)
 	}
 	else if (Cli.got_subcommand(SubInfo))
 	{
-		return CmdInfo(InputFilename, InputFilename2, bInfoFiles);
+		FCmdInfoOptions Options;
+		Options.InputA = InputFilename;
+		Options.InputB = InputFilename2;
+		Options.bListFiles = bInfoFiles;
+		Options.SyncFilter = &SyncFilter;
+		return CmdInfo(Options);
 	}
 	else if (Cli.got_subcommand(SubQuery))
 	{
