unsync - Groundwork for mount command

diff --git a/Engine/Source/Programs/Unsync/Private/UnsyncCmdMount.cpp b/Engine/Source/Programs/Unsync/Private/UnsyncCmdMount.cpp
new file mode 100644
--- /dev/null
+++ b/Engine/Source/Programs/Unsync/Private/UnsyncCmdMount.cpp
@@ -0,0 +1,18 @@
+// Copyright Epic Games, Inc. All Rights Reserved.
+
+#include "UnsyncCmdMount.h"
+#include "UnsyncMount.h"
+
+namespace unsync {
+
+int32
+CmdMount(const FCmdMountOptions& Options)
+{
+	FMountedDirectory MountedDir;
+
+	bool bOk = MountedDir.Mount(Options.Path);
+
+	return bOk ? 0 : 1;
+}
+
+}  // namespace unsync
diff --git a/Engine/Source/Programs/Unsync/Private/UnsyncCmdMount.h b/Engine/Source/Programs/Unsync/Private/UnsyncCmdMount.h
new file mode 100644
--- /dev/null
+++ b/Engine/Source/Programs/Unsync/Private/UnsyncCmdMount.h
@@ -0,0 +1,16 @@
+// Copyright Epic Games, Inc. All Rights Reserved.
+
+#pragma once
+
+#include "UnsyncCommon.h"
+
+namespace unsync {
+
+struct FCmdMountOptions
+{
+	FPath Path;
+};
+
+int32 CmdMount(const FCmdMountOptions& Options);
+
+}
diff --git a/Engine/Source/Programs/Unsync/Private/UnsyncMain.cpp b/Engine/Source/Programs/Unsync/Private/UnsyncMain.cpp
--- a/Engine/Source/Programs/Unsync/Private/UnsyncMain.cpp
+++ b/Engine/Source/Programs/Unsync/Private/UnsyncMain.cpp
@@ -6,6 +6,7 @@
 #include "UnsyncCmdPush.h"
 #include "UnsyncCmdSync.h"
 #include "UnsyncCmdQuery.h"
+#include "UnsyncCmdMount.h"
 #include "UnsyncCore.h"
 #include "UnsyncFile.h"
 #include "UnsyncMemory.h"
@@ -103,6 +104,12 @@ InnerMain(int Argc, char** Argv)
 	};
 #endif // UNSYNC_USE_TLS
 
+	auto AddProxyOptions = [&RemoteAddressUtf8](CLI::App* App) {
+		App->add_option("--proxy, --remote",
+							RemoteAddressUtf8,
+							"FProxy server address ([transport://]address[:port][/request][#namespace])");
+	};
+
 	CLI::App* SubHash = Cli.add_subcommand("hash", "Generate hash manifest for a file or directory");
 	SubHash->add_option("Input", InputFilenameUtf8, "Input file or directory path")->required();
 	SubHash->add_flag("-f, --force", bForceOperation, "Force the operation even if hash is already computed for the input");
@@ -171,9 +178,7 @@ InnerMain(int Argc, char** Argv)
 		->required();
 	SubSync->add_option("Target", TargetFilenameUtf8, "Target path")->required();
 	SubSync->add_option("-m, --manifest", SourceManifestFilenameUtf8, "Override manifest path for Source");
-	SubSync->add_option("--proxy, --remote",
-						RemoteAddressUtf8,
-						"FProxy server address ([transport://]address[:port][/request][#namespace])");
+	AddProxyOptions(SubSync);
 	SubSync->add_option("--dfs", PreferredDfsUtf8, "Preferred DFS mirror (matched by sub-string)");
 	SubSync->add_option("--overlay", OverlayArrayUtf8, "Additional source directory to sync (keep unique files from all sources, overwrite conflicting files with overlay source)");
 	SubSync->add_option("--include", IncludeFilterArrayUtf8, "Include filenames that contain specified words (comma separated). If this is not present, all files will be included.");
@@ -242,14 +247,22 @@ InnerMain(int Argc, char** Argv)
 
 	CLI::App* SubQuery = Cli.add_subcommand("query", "Run a query command on the remote server");
 	SubQuery->add_option("QueryString", QueryStringUtf8, "Query")->required();
-	SubQuery->add_option("--proxy, --remote",
-					RemoteAddressUtf8,
-					"FProxy server address ([transport://]address[:port][/request][#namespace])")->required();
+	AddProxyOptions(SubQuery);
+
 #if UNSYNC_USE_TLS
 	AddTlsOptions(SubQuery);
 #endif // UNSYNC_USE_TLS
 	SubCommands.push_back(SubQuery);
 
+	CLI::App* SubMount = Cli.add_subcommand("mount", "Mount directory manifest as a virtual file system (EXPERIMENTAL)");
+	SubMount
+		->add_option("Source",
+					 SourceFilenameUtf8,
+					 "Source path, object name, hash or full URL ([transport://]address[:port]#namespace/object)")
+		->required();
+	AddProxyOptions(SubMount);
+	SubCommands.push_back(SubMount);
+
 	for (CLI::App* Subcommand : SubCommands)
 	{
 		Subcommand->add_flag("-d, --dry, --dry-run", GDryRun, "Don't write any outputs to disk");
@@ -723,6 +736,12 @@ InnerMain(int Argc, char** Argv)
 		QueryOptions.Remote = RemoteDesc;
 		return CmdQuery(QueryOptions);
 	}
+	else if (Cli.got_subcommand(SubMount))
+	{
+		FCmdMountOptions MountOptions;
+		MountOptions.Path = SourceFilename;
+		return CmdMount(MountOptions);
+	}
 
 	return 0;
 }
diff --git a/Engine/Source/Programs/Unsync/Private/UnsyncMount.cpp b/Engine/Source/Programs/Unsync/Private/UnsyncMount.cpp
new file mode 100644
--- /dev/null
+++ b/Engine/Source/Programs/Unsync/Private/UnsyncMount.cpp
@@ -0,0 +1,122 @@
+// Copyright Epic Games, Inc. All Rights Reserved.
+
+#include "UnsyncMount.h"
+#include "UnsyncCore.h"
+#include "UnsyncSerialization.h"
+
+#include <algorithm>
+
+namespace unsync {
+
+static int32
+GetPathDepth(std::wstring_view S)
+{
+	int32 Result = 0;
+	for (wchar_t c : S)
+	{
+		if (c == L'/' || c == L'\\')
+		{
+			Result++;
+		}
+	}
+	return Result;
+}
+
+std::wstring_view
+GetParentDirectory(std::wstring_view Name)
+{
+	size_t Pos = Name.find_last_of(L"\\/");
+	if (Pos == std::string::npos)
+	{
+		return {};
+	}
+	else
+	{
+		return Name.substr(0, Pos);
+	}
+}
+
+bool
+FMountedDirectory::Mount(const FPath& RootPath)
+{
+	FPath DirectoryManifestPath = RootPath / ".unsync" / "manifest.bin";
+
+	bValid = false;
+	DirectoryEntries.clear();
+
+	if (!LoadDirectoryManifest(Manifest, RootPath, DirectoryManifestPath))
+	{
+		UNSYNC_ERROR("Failed to load directory manifest");
+		return false;
+	}
+
+	std::unordered_set<std::wstring_view> SubdirectorySet;
+
+	for (const auto& it : Manifest.Files)
+	{
+		FDirectoryEntry FileEntry;
+		FileEntry.Type	   = EDirectoryEntryType::File;
+		FileEntry.Path	   = it.first;
+		FileEntry.Manifest = &it.second;
+		FileEntry.Depth	   = GetPathDepth(FileEntry.Path);
+
+		DirectoryEntries.push_back(FileEntry);
+
+		std::wstring_view ParentName = GetParentDirectory(FileEntry.Path);
+
+		while (!ParentName.empty())
+		{
+			auto InsertResult = SubdirectorySet.insert(ParentName);
+
+			if (InsertResult.second)
+			{
+				DirectoryEntries.back().Parent = *InsertResult.first;
+
+				FDirectoryEntry SubdirEntry;
+				SubdirEntry.Type	 = EDirectoryEntryType::Directory;
+				SubdirEntry.Path	 = ParentName;
+				SubdirEntry.Manifest = nullptr;
+				SubdirEntry.Depth	 = GetPathDepth(SubdirEntry.Path);
+
+				DirectoryEntries.push_back(SubdirEntry);
+			}
+			else
+			{
+				if (DirectoryEntries.back().Parent.empty())
+				{
+					DirectoryEntries.back().Parent = *InsertResult.first;
+				}
+			}
+
+			ParentName = GetParentDirectory(ParentName);
+		}
+	}
+
+	std::sort(DirectoryEntries.begin(), DirectoryEntries.end());
+	std::wstring_view LastParent;
+	for (uint64 i = 0; i < DirectoryEntries.size(); ++i)
+	{
+		FDirectoryEntry& Entry = DirectoryEntries[i];
+
+		if (Entry.Parent.empty())
+		{
+			Entry.Name = Entry.Path;
+		}
+		else
+		{
+			Entry.Name = Entry.Path.substr(Entry.Parent.length() + 1);
+		}
+
+		if (Entry.Parent != LastParent)
+		{
+			DirectoryMap[Entry.Parent] = i;
+			LastParent = Entry.Parent;
+		}
+	}
+
+	bValid = true;
+
+	return true;
+}
+}
+
diff --git a/Engine/Source/Programs/Unsync/Private/UnsyncMount.h b/Engine/Source/Programs/Unsync/Private/UnsyncMount.h
new file mode 100644
--- /dev/null
+++ b/Engine/Source/Programs/Unsync/Private/UnsyncMount.h
@@ -0,0 +1,57 @@
+// Copyright Epic Games, Inc. All Rights Reserved.
+
+#pragma once
+
+#include "UnsyncCommon.h"
+#include "UnsyncCore.h"
+
+namespace unsync {
+
+// Represents a VFS-like view of a directory manifest, providing API to list files,
+struct FMountedDirectory
+{
+	UNSYNC_DISALLOW_COPY_ASSIGN(FMountedDirectory)
+
+	FMountedDirectory() = default;
+
+	bool			   bValid = false;
+	FDirectoryManifest Manifest;
+
+	enum class EDirectoryEntryType {
+		// Filesystem entries are sorted by this and we want to list directories first during iteration
+		Directory = 0,
+		File	  = 1,
+	};
+
+	struct FDirectoryEntry
+	{
+		int32				Depth = 0;
+		EDirectoryEntryType Type  = EDirectoryEntryType::File;
+		// Full path relative to mounted directory root
+		std::wstring_view Path;
+		// Full path of parent directory (guaranteed to be deduplicated, i.e. identical strings point to the same memory)
+		std::wstring_view Parent;
+		// File/directory name (without parent path)
+		std::wstring_view	 Name;
+		const FFileManifest* Manifest = nullptr;
+
+		bool operator<(const FDirectoryEntry& Other) const
+		{
+			if (Depth != Other.Depth)
+				return Depth < Other.Depth;
+			else if (Type != Other.Type)
+				return Type < Other.Type;
+			else
+				return Path < Other.Path;
+		}
+	};
+
+	std::vector<FDirectoryEntry> DirectoryEntries;
+
+	// Offset of the first directory entry corresponding to a speficic parent name
+	std::unordered_map<std::wstring_view, uint64> DirectoryMap;
+
+	bool Mount(const FPath& RootPath);
+};
+
+}  // namespace unsync
diff --git a/Engine/Source/Programs/Unsync/Private/UnsyncUtil.h b/Engine/Source/Programs/Unsync/Private/UnsyncUtil.h
--- a/Engine/Source/Programs/Unsync/Private/UnsyncUtil.h
+++ b/Engine/Source/Programs/Unsync/Private/UnsyncUtil.h
@@ -176,6 +176,13 @@ MakeView(const T* Ptr, size_t Count)
 	return TArrayView{Ptr, Ptr + Count};
 }
 
+template<typename T>
+TArrayView<T>
+MakeView(const std::vector<T>& Container)
+{
+	return MakeView(Container.data(), Container.size());
+}
+
 inline uint64
 AlignDownToMultiplePow2(uint64 X, uint64 MultiplePow2)
 {
