unsync - Automatically find closest proxy server by default

* Use --no-proxy-select to skip server selection
* Bump version to 1.0.52

diff --git a/Engine/Source/Programs/Unsync/Private/UnsyncCmdQuery.cpp b/Engine/Source/Programs/Unsync/Private/UnsyncCmdQuery.cpp
--- a/Engine/Source/Programs/Unsync/Private/UnsyncCmdQuery.cpp
+++ b/Engine/Source/Programs/Unsync/Private/UnsyncCmdQuery.cpp
@@ -5,19 +5,12 @@
 #include "UnsyncThread.h"
 #include "UnsyncUtil.h"
 
+#include <float.h>
 #include <algorithm>
 #include <json11.hpp>
-#include <float.h>
 
 namespace unsync {
 
-struct FMirrorInfo
-{
-	std::string Name;
-	std::string Address;
-	uint16		Port = UNSYNC_DEFAULT_PORT;
-	double		Ping = 0;
-};
 using FMirrorInfoResult = TResult<std::vector<FMirrorInfo>>;
 
 // Runs a basic HTTP request against the remote server and returns the time it took to get the response, -1 if connection could not be
@@ -46,7 +39,7 @@ RunHttpPing(std::string_view Address, uint16 Port)
 FMirrorInfoResult
 RunQueryMirrors(const FRemoteDesc& RemoteDesc)
 {
-	const char* Url = "/api/v1/mirrors";
+	const char*	  Url	   = "/api/v1/mirrors";
 	FHttpResponse Response = HttpRequest(RemoteDesc, EHttpMethod::GET, Url);
 
 	if (!Response.Success())
@@ -151,4 +144,33 @@ CmdQuery(const FCmdQueryOptions& Options)
 	}
 }
 
+TResult<FMirrorInfo>
+FindClosestMirror(const FRemoteDesc& Remote)
+{
+	FMirrorInfoResult MirrorsResult = RunQueryMirrors(Remote);
+	if (MirrorsResult.IsError())
+	{
+		return FError(MirrorsResult.GetError());
+	}
+
+	std::vector<FMirrorInfo> Mirrors = MirrorsResult.GetData();
+	ParallelForEach(Mirrors.begin(), Mirrors.end(), [](FMirrorInfo& Mirror) { Mirror.Ping = RunHttpPing(Mirror.Address, Mirror.Port); });
+
+	std::sort(Mirrors.begin(), Mirrors.end(), [](const FMirrorInfo& InA, const FMirrorInfo& InB) {
+		double A = InA.Ping > 0 ? InA.Ping : FLT_MAX;
+		double B = InB.Ping > 0 ? InB.Ping : FLT_MAX;
+		return A < B;
+	});
+
+	for (const FMirrorInfo& Mirror : Mirrors)
+	{
+		if (Mirror.Ping > 0)
+		{
+			return ResultOk(Mirror);
+		}
+	}
+
+	return AppError("No reachable mirror found");
+}
+
 }  // namespace unsync
diff --git a/Engine/Source/Programs/Unsync/Private/UnsyncCmdQuery.h b/Engine/Source/Programs/Unsync/Private/UnsyncCmdQuery.h
--- a/Engine/Source/Programs/Unsync/Private/UnsyncCmdQuery.h
+++ b/Engine/Source/Programs/Unsync/Private/UnsyncCmdQuery.h
@@ -15,4 +15,14 @@ struct FCmdQueryOptions
 
 int32 CmdQuery(const FCmdQueryOptions& Options);
 
+struct FMirrorInfo
+{
+	std::string Name;
+	std::string Address;
+	uint16		Port = UNSYNC_DEFAULT_PORT;
+	double		Ping = 0;
+};
+
+TResult<FMirrorInfo> FindClosestMirror(const FRemoteDesc& Remote);
+
 } // namespace unsync
diff --git a/Engine/Source/Programs/Unsync/Private/UnsyncCore.cpp b/Engine/Source/Programs/Unsync/Private/UnsyncCore.cpp
--- a/Engine/Source/Programs/Unsync/Private/UnsyncCore.cpp
+++ b/Engine/Source/Programs/Unsync/Private/UnsyncCore.cpp
@@ -27,7 +27,7 @@ UNSYNC_THIRD_PARTY_INCLUDES_START
 #include <md5-sse2.h>
 UNSYNC_THIRD_PARTY_INCLUDES_END
 
-#define UNSYNC_VERSION_STR "1.0.52-dev"
+#define UNSYNC_VERSION_STR "1.0.52"
 
 namespace unsync {
 
diff --git a/Engine/Source/Programs/Unsync/Private/UnsyncMain.cpp b/Engine/Source/Programs/Unsync/Private/UnsyncMain.cpp
--- a/Engine/Source/Programs/Unsync/Private/UnsyncMain.cpp
+++ b/Engine/Source/Programs/Unsync/Private/UnsyncMain.cpp
@@ -2,11 +2,11 @@
 
 #include "UnsyncCmdDiff.h"
 #include "UnsyncCmdHash.h"
+#include "UnsyncCmdMount.h"
 #include "UnsyncCmdPatch.h"
 #include "UnsyncCmdPush.h"
-#include "UnsyncCmdSync.h"
 #include "UnsyncCmdQuery.h"
-#include "UnsyncCmdMount.h"
+#include "UnsyncCmdSync.h"
 #include "UnsyncCore.h"
 #include "UnsyncFile.h"
 #include "UnsyncMemory.h"
@@ -21,11 +21,11 @@ UNSYNC_THIRD_PARTY_INCLUDES_START
 #	include <shellapi.h>
 #endif	// UNSYNC_PLATFORM_WINDOWS
 #include <fcntl.h>
+#include <CLI/CLI.hpp>
 #include <filesystem>
 #include <fstream>
 #include <iostream>
 #include <system_error>
-#include <CLI/CLI.hpp>
 UNSYNC_THIRD_PARTY_INCLUDES_END
 
 namespace unsync {
@@ -82,32 +82,35 @@ InnerMain(int Argc, char** Argv)
 	bool					 bFullSourceScan	 = false;
 	bool					 bFullDifference	 = false;
 	bool					 bInfoFiles			 = false;
+	bool					 bNoProxySelect		 = false;
 	int32					 CompressionLevel	 = 3;
 	uint32					 DiffBlockSize		 = uint32(4_KB);
 	uint32					 HashOrSyncBlockSize = uint32(64_KB);
 
 	struct FDeprecatedOptions
 	{
-		bool bQuickSyncMode = false;
-		bool bQuickDifference = false;
+		bool bQuickSyncMode			= false;
+		bool bQuickDifference		= false;
 		bool bQuickSourceValidation = false;
 	} DeprecatedOptions;
 
-	const std::string HiddenGroupId; // CLI11 uses an empty string group name to mark arguments that should be hidden
+	const std::string HiddenGroupId;  // CLI11 uses an empty string group name to mark arguments that should be hidden
 
 #if UNSYNC_USE_TLS
-	auto AddTlsOptions = [&CacertFilenameUtf8, &bUseTls, &bAllowInsecureTls](CLI::App* App)
-	{
+	auto AddTlsOptions = [&CacertFilenameUtf8, &bUseTls, &bAllowInsecureTls](CLI::App* App) {
 		App->add_option("--cacert", CacertFilenameUtf8, "Certificate authority file to use for TLS validation (.pem)");
 		App->add_flag("--tls", bUseTls, "Use TLS when connecting to remote server");
 		App->add_flag("--insecure", bAllowInsecureTls, "Skip remote server TLS certificate validation");
 	};
-#endif // UNSYNC_USE_TLS
+#endif	// UNSYNC_USE_TLS
 
-	auto AddProxyOptions = [&RemoteAddressUtf8](CLI::App* App) {
+	auto AddProxyOptions = [&RemoteAddressUtf8, &bNoProxySelect](CLI::App* App) {
 		App->add_option("--proxy, --remote",
-							RemoteAddressUtf8,
-							"FProxy server address ([transport://]address[:port][/request][#namespace])");
+						RemoteAddressUtf8,
+						"FProxy server address ([transport://]address[:port][/request][#namespace])");
+		App->add_option("--no-proxy-select",
+						bNoProxySelect,
+						"Skip automatic server selection and use the exact one specified by command line or environment variable");
 	};
 
 	CLI::App* SubHash = Cli.add_subcommand("hash", "Generate hash manifest for a file or directory");
@@ -180,9 +183,17 @@ InnerMain(int Argc, char** Argv)
 	SubSync->add_option("-m, --manifest", SourceManifestFilenameUtf8, "Override manifest path for Source");
 	AddProxyOptions(SubSync);
 	SubSync->add_option("--dfs", PreferredDfsUtf8, "Preferred DFS mirror (matched by sub-string)");
-	SubSync->add_option("--overlay", OverlayArrayUtf8, "Additional source directory to sync (keep unique files from all sources, overwrite conflicting files with overlay source)");
-	SubSync->add_option("--include", IncludeFilterArrayUtf8, "Include filenames that contain specified words (comma separated). If this is not present, all files will be included.");
-	SubSync->add_option("--exclude", ExcludeFilterArrayUtf8, "Exclude filenames that contain specified words (comma separated). Filter is run after --include.");
+	SubSync->add_option(
+		"--overlay",
+		OverlayArrayUtf8,
+		"Additional source directory to sync (keep unique files from all sources, overwrite conflicting files with overlay source)");
+	SubSync->add_option(
+		"--include",
+		IncludeFilterArrayUtf8,
+		"Include filenames that contain specified words (comma separated). If this is not present, all files will be included.");
+	SubSync->add_option("--exclude",
+						ExcludeFilterArrayUtf8,
+						"Exclude filenames that contain specified words (comma separated). Filter is run after --include.");
 #if UNSYNC_USE_TLS
 	AddTlsOptions(SubSync);
 #else
@@ -192,7 +203,9 @@ InnerMain(int Argc, char** Argv)
 						HttpHeaderFilenameUtf8,
 						"Text file that contains any extra HTTP headers to pass to the remote server (auth tokens, etc.)");
 	SubSync->add_flag("--no-cleanup", bNoCleanupAfterSync, "Do not delete local files that aren't in the manifest after a successful sync");
-	SubSync->add_option("--cleanup-exclude", CleanupExcludeFilterArrayUtf8, "Exclude filenames that contain specified words from cleanup process (comma separated)");
+	SubSync->add_option("--cleanup-exclude",
+						CleanupExcludeFilterArrayUtf8,
+						"Exclude filenames that contain specified words from cleanup process (comma separated)");
 
 	// Deprecated --quick flag
 	SubSync
@@ -251,7 +264,7 @@ InnerMain(int Argc, char** Argv)
 
 #if UNSYNC_USE_TLS
 	AddTlsOptions(SubQuery);
-#endif // UNSYNC_USE_TLS
+#endif	// UNSYNC_USE_TLS
 	SubCommands.push_back(SubQuery);
 
 	CLI::App* SubMount = Cli.add_subcommand("mount", "Mount directory manifest as a virtual file system (EXPERIMENTAL)");
@@ -622,8 +635,8 @@ InnerMain(int Argc, char** Argv)
 	}
 
 	{
-		FPath ExtraCertPath = GExePath.parent_path() / "unsync.cer";
-		FBuffer CertBuffer = ReadFileToBuffer(ExtraCertPath);
+		FPath	ExtraCertPath = GExePath.parent_path() / "unsync.cer";
+		FBuffer CertBuffer	  = ReadFileToBuffer(ExtraCertPath);
 		if (!CertBuffer.Empty())
 		{
 			UNSYNC_LOG(L"Using trusted certificates from '%ls'", ExtraCertPath.wstring().c_str());
@@ -639,6 +652,25 @@ InnerMain(int Argc, char** Argv)
 		}
 	}
 
+	if (!bNoProxySelect && RemoteDesc.IsValid() && RemoteDesc.Protocol == EProtocolFlavor::Unsync)
+	{
+		UNSYNC_LOG(L"Selecting server using root '%hs'", RemoteDesc.HostAddress.c_str());
+		TResult<FMirrorInfo> MirrorResult = FindClosestMirror(RemoteDesc);
+		if (const FMirrorInfo* Mirror = MirrorResult.TryData())
+		{
+			UNSYNC_LOG(L"Closest server: '%hs', ping: %.2f ms", Mirror->Address.c_str(), Mirror->Ping * 1000.0);
+
+			RemoteDesc.HostAddress = Mirror->Address;
+			RemoteDesc.HostPort	   = Mirror->Port;
+		}
+		else
+		{
+			UNSYNC_WARNING(L"Failed to find closest proxy using root server '%hs': %ls",
+						   RemoteAddressUtf8.c_str(),
+						   MirrorResult.TryError()->Context.c_str());
+		}
+	}
+
 	if (Cli.got_subcommand(SubHash))
 	{
 		FCmdHashOptions HashOptions;
@@ -723,8 +755,8 @@ InnerMain(int Argc, char** Argv)
 	else if (Cli.got_subcommand(SubInfo))
 	{
 		FCmdInfoOptions Options;
-		Options.InputA = InputFilename;
-		Options.InputB = InputFilename2;
+		Options.InputA	   = InputFilename;
+		Options.InputB	   = InputFilename2;
 		Options.bListFiles = bInfoFiles;
 		Options.SyncFilter = &SyncFilter;
 		return CmdInfo(Options);
@@ -827,7 +859,7 @@ main(int argc, char** argv)
 	{
 		ArgvUtf8.push_back(ArgvStringsUtf8[I].data());
 	}
-#else // UNSYNC_PLATFORM_WINDOWS
+#else	// UNSYNC_PLATFORM_WINDOWS
 	GExePath = FPath(argv[0]);
 #endif	// UNSYNC_PLATFORM_WINDOWS
 
