unsync - Skip automatic proxy selection for query command, use dedicated ping endpoint, refactor query helpers and download result mechanism

diff --git a/Engine/Source/Programs/Unsync/CMakeLists.txt b/Engine/Source/Programs/Unsync/CMakeLists.txt
--- a/Engine/Source/Programs/Unsync/CMakeLists.txt
+++ b/Engine/Source/Programs/Unsync/CMakeLists.txt
@@ -37,6 +37,8 @@ set(src
 	Private/UnsyncCmdDiff.h
 	Private/UnsyncCmdHash.cpp
 	Private/UnsyncCmdHash.h
+	Private/UnsyncCmdMount.cpp
+	Private/UnsyncCmdMount.h
 	Private/UnsyncCmdPatch.cpp
 	Private/UnsyncCmdPatch.h
 	Private/UnsyncCmdPush.cpp
@@ -67,6 +69,8 @@ set(src
 	Private/UnsyncMemory.h
 	Private/UnsyncMiniCb.cpp
 	Private/UnsyncMiniCb.h
+	Private/UnsyncMount.cpp
+	Private/UnsyncMount.h
 	Private/UnsyncPool.h
 	Private/UnsyncProgress.cpp
 	Private/UnsyncProgress.h
diff --git a/Engine/Source/Programs/Unsync/Private/UnsyncCmdQuery.cpp b/Engine/Source/Programs/Unsync/Private/UnsyncCmdQuery.cpp
--- a/Engine/Source/Programs/Unsync/Private/UnsyncCmdQuery.cpp
+++ b/Engine/Source/Programs/Unsync/Private/UnsyncCmdQuery.cpp
@@ -21,7 +21,7 @@ RunHttpPing(std::string_view Address, uint16 Port)
 	FTimePoint		TimeBegin = TimePointNow();
 	FHttpConnection Connection(Address, Port);
 	FHttpRequest	Request;
-	Request.Url				   = "/api/v1/hello";
+	Request.Url				   = "/api/v1/ping";
 	Request.Method			   = EHttpMethod::GET;
 	FHttpResponse PingResponse = HttpRequest(Connection, Request);
 	FTimePoint	  TimeEnd	   = TimePointNow();
diff --git a/Engine/Source/Programs/Unsync/Private/UnsyncCore.cpp b/Engine/Source/Programs/Unsync/Private/UnsyncCore.cpp
--- a/Engine/Source/Programs/Unsync/Private/UnsyncCore.cpp
+++ b/Engine/Source/Programs/Unsync/Private/UnsyncCore.cpp
@@ -27,7 +27,7 @@ UNSYNC_THIRD_PARTY_INCLUDES_START
 #include <md5-sse2.h>
 UNSYNC_THIRD_PARTY_INCLUDES_END
 
-#define UNSYNC_VERSION_STR "1.0.52"
+#define UNSYNC_VERSION_STR "1.0.53-dev"
 
 namespace unsync {
 
@@ -1186,7 +1186,7 @@ DownloadBlocks(FProxyPool&					  ProxyPool,
 					{
 						ProxyPool.Dealloc(std::move(Proxy));
 					}
-					else if (DownloadResult.GetError() == EDownloadRetryMode::Abort)
+					else if (DownloadResult.GetError().RetryMode == EDownloadRetryMode::Abort)
 					{
 						bGotError = true;
 						ProxyPool.Invalidate();
diff --git a/Engine/Source/Programs/Unsync/Private/UnsyncJupiter.cpp b/Engine/Source/Programs/Unsync/Private/UnsyncJupiter.cpp
--- a/Engine/Source/Programs/Unsync/Private/UnsyncJupiter.cpp
+++ b/Engine/Source/Programs/Unsync/Private/UnsyncJupiter.cpp
@@ -823,12 +823,12 @@ FJupiterProtocolImpl::Download(const TArrayView<FNeedBlock> NeedBlocks, const FB
 {
 	if (!IsValid())
 	{
-		return FDownloadResult(EDownloadRetryMode::Abort);
+		return FDownloadError(EDownloadRetryMode::Abort);
 	}
 
 	if (NeedBlocks.Size() == 0)
 	{
-		return ResultOk<EDownloadRetryMode>();
+		return ResultOk<FDownloadError>();
 	}
 
 	UNSYNC_ASSERT(Connection.NumActiveRequests == 0);
@@ -1044,7 +1044,7 @@ FJupiterProtocolImpl::Download(const TArrayView<FNeedBlock> NeedBlocks, const FB
 	UNSYNC_ASSERT(Connection.NumActiveRequests == 0);
 	UNSYNC_ASSERT(Connection.ResponseQueue.empty());
 
-	return ResultOk<EDownloadRetryMode>();
+	return ResultOk<FDownloadError>();
 }
 
 void
diff --git a/Engine/Source/Programs/Unsync/Private/UnsyncMain.cpp b/Engine/Source/Programs/Unsync/Private/UnsyncMain.cpp
--- a/Engine/Source/Programs/Unsync/Private/UnsyncMain.cpp
+++ b/Engine/Source/Programs/Unsync/Private/UnsyncMain.cpp
@@ -108,9 +108,9 @@ InnerMain(int Argc, char** Argv)
 		App->add_option("--proxy, --remote",
 						RemoteAddressUtf8,
 						"FProxy server address ([transport://]address[:port][/request][#namespace])");
-		App->add_option("--no-proxy-select",
-						bNoProxySelect,
-						"Skip automatic server selection and use the exact one specified by command line or environment variable");
+		App->add_flag("--no-proxy-select",
+					  bNoProxySelect,
+					  "Skip automatic server selection and use the exact one specified by command line or environment variable");
 	};
 
 	CLI::App* SubHash = Cli.add_subcommand("hash", "Generate hash manifest for a file or directory");
@@ -652,7 +652,8 @@ InnerMain(int Argc, char** Argv)
 		}
 	}
 
-	if (!bNoProxySelect && RemoteDesc.IsValid() && RemoteDesc.Protocol == EProtocolFlavor::Unsync)
+	if (!bNoProxySelect && !Cli.got_subcommand(SubQuery)
+		&& RemoteDesc.IsValid() && RemoteDesc.Protocol == EProtocolFlavor::Unsync)
 	{
 		UNSYNC_LOG(L"Selecting server using root '%hs'", RemoteDesc.HostAddress.c_str());
 		TResult<FMirrorInfo> MirrorResult = FindClosestMirror(RemoteDesc);
diff --git a/Engine/Source/Programs/Unsync/Private/UnsyncProtocol.h b/Engine/Source/Programs/Unsync/Private/UnsyncProtocol.h
--- a/Engine/Source/Programs/Unsync/Private/UnsyncProtocol.h
+++ b/Engine/Source/Programs/Unsync/Private/UnsyncProtocol.h
@@ -129,7 +129,7 @@ struct FCommandPacket
 	uint64					CommandId = 0;
 };
 
-struct FileListPacket
+struct FFileListPacket
 {
 	static constexpr uint64 MAGIC		  = 0x28B96050A327172Aull;
 	uint64					Magic		  = MAGIC;
diff --git a/Engine/Source/Programs/Unsync/Private/UnsyncProxy.cpp b/Engine/Source/Programs/Unsync/Private/UnsyncProxy.cpp
--- a/Engine/Source/Programs/Unsync/Private/UnsyncProxy.cpp
+++ b/Engine/Source/Programs/Unsync/Private/UnsyncProxy.cpp
@@ -30,16 +30,6 @@ struct FUnsyncProtocolImpl : FRemoteProtocolBase
 
 	const FRemoteProtocolFeatures Features;
 
-	struct FHelloResponse
-	{
-		std::string Name;
-		std::string VersionNumber;
-		std::string VersionGit;
-		std::string SessionId;
-		std::vector<std::string> FeatureNames;
-		FRemoteProtocolFeatures Features;
-	};
-	static TResult<FHelloResponse> QueryHello(const FRemoteDesc& RemoteDesc);
 	static void SendTelemetryEvent(const FRemoteDesc& RemoteDesc, const FTelemetryEventSyncComplete& Event);
 };
 
@@ -241,7 +231,7 @@ FUnsyncProtocolImpl::Download(const TArrayView<FNeedBlock> NeedBlocks, const FBl
 			}
 		}
 
-		FileListPacket FileListHeader;
+		FFileListPacket FileListHeader;
 		FileListHeader.DataSizeBytes = CheckedNarrow(FileListData.Size());
 		FileListHeader.NumFiles		 = CheckedNarrow(FileListUtf8.size());
 
@@ -304,8 +294,11 @@ FUnsyncProtocolImpl::Download(const TArrayView<FNeedBlock> NeedBlocks, const FBl
 
 		bOk &= (SocketRecvAll(*SocketHandle, BlockPacket.CompressedData.Data(), BlockPacket.CompressedData.Size()) == CompressedDataSize);
 
-		if (BlockPacket.Hash == FHash128{})	 // response is always terminated with an empty packet
+		static const FHash128 TerminatorHash = FHash128{}; // response is always terminated with an empty packet
+
+		if (BlockPacket.Hash == TerminatorHash)
 		{
+			// TODO: Termination packet payload may contain diagnostic messages that can be reported to the user
 			break;
 		}
 
@@ -327,7 +320,7 @@ FUnsyncProtocolImpl::Download(const TArrayView<FNeedBlock> NeedBlocks, const FBl
 		bIsConnetedToHost = false;
 	}
 
-	return ResultOk<EDownloadRetryMode>();
+	return ResultOk<FDownloadError>();
 }
 
 void
@@ -350,7 +343,8 @@ FUnsyncProtocolImpl::GetSocketSecurity() const
 	}
 }
 
-TResult<FUnsyncProtocolImpl::FHelloResponse> FUnsyncProtocolImpl::QueryHello(const FRemoteDesc& RemoteDesc)
+namespace ProxyQuery {
+TResult<FHelloResponse> Hello(const FRemoteDesc& RemoteDesc)
 {
 	const char* Url = "/api/v1/hello";
 	FHttpResponse Response = HttpRequest(RemoteDesc, EHttpMethod::GET, Url);
@@ -426,6 +420,7 @@ TResult<FUnsyncProtocolImpl::FHelloResponse> FUnsyncProtocolImpl::QueryHello(con
 
 	return ResultOk(std::move(Result));
 }
+}  // namespace ProxyQuery
 
 void
 FUnsyncProtocolImpl::SendTelemetryEvent(const FRemoteDesc& RemoteDesc, const FTelemetryEventSyncComplete& Event)
@@ -596,7 +591,7 @@ FProxyPool::FProxyPool(const FRemoteDesc& InRemoteDesc)
 			RemoteDesc.HostAddress.c_str(),
 			RemoteDesc.HostPort);
 
-		TResult<FUnsyncProtocolImpl::FHelloResponse> Response = FUnsyncProtocolImpl::QueryHello(RemoteDesc);
+		TResult<ProxyQuery::FHelloResponse> Response = ProxyQuery::Hello(RemoteDesc);
 
 		if (Response.IsError())
 		{
@@ -604,7 +599,7 @@ FProxyPool::FProxyPool(const FRemoteDesc& InRemoteDesc)
 		}
 		else
 		{
-			const FUnsyncProtocolImpl::FHelloResponse& Data = Response.GetData();
+			const ProxyQuery::FHelloResponse& Data = Response.GetData();
 			UNSYNC_VERBOSE(L"Connection established. Server name: %hs, version: %hs, git: %hs.", 
 				Data.Name.empty() ? "unknown" : Data.Name.c_str(),
 				Data.VersionNumber.empty() ? "unknown" : Data.VersionNumber.c_str(),
diff --git a/Engine/Source/Programs/Unsync/Private/UnsyncProxy.h b/Engine/Source/Programs/Unsync/Private/UnsyncProxy.h
--- a/Engine/Source/Programs/Unsync/Private/UnsyncProxy.h
+++ b/Engine/Source/Programs/Unsync/Private/UnsyncProxy.h
@@ -23,7 +23,17 @@ enum class EDownloadRetryMode
 	Abort,	// unrecoverable error
 };
 
-using FDownloadResult = TResult<FEmpty, EDownloadRetryMode>;
+struct FDownloadError : FError
+{
+	FDownloadError() = default;
+	FDownloadError(EDownloadRetryMode InRetryMode) : RetryMode(InRetryMode) {}
+
+	EDownloadRetryMode RetryMode = EDownloadRetryMode::Abort;
+
+	bool CanRetry() const { return RetryMode == EDownloadRetryMode ::Retry; }
+};
+
+using FDownloadResult = TResult<FEmpty, FDownloadError>;
 
 struct FDownloadedBlock
 {
@@ -170,4 +180,21 @@ private:
 	std::mutex Mutex;
 };
 
+namespace ProxyQuery
+{
+
+struct FHelloResponse
+{
+	std::string Name;
+	std::string VersionNumber;
+	std::string VersionGit;
+	std::string SessionId;
+
+	std::vector<std::string> FeatureNames;
+	FRemoteProtocolFeatures	 Features;
+};
+TResult<FHelloResponse> Hello(const FRemoteDesc& RemoteDesc);
+
+} 
+
 }  // namespace unsync
