unsync - Groundwork for user authentication

diff --git a/Engine/Source/Programs/Unsync/CMakeLists.txt b/Engine/Source/Programs/Unsync/CMakeLists.txt
--- a/Engine/Source/Programs/Unsync/CMakeLists.txt
+++ b/Engine/Source/Programs/Unsync/CMakeLists.txt
@@ -31,6 +31,8 @@ endif()
 # Add unsync code
 
 set(src
+	Private/UnsyncAuth.cpp
+	Private/UnsyncAuth.h
 	Private/UnsyncBuffer.cpp
 	Private/UnsyncBuffer.h
 	Private/UnsyncCmdDiff.cpp
diff --git a/Engine/Source/Programs/Unsync/Private/UnsyncAuth.cpp b/Engine/Source/Programs/Unsync/Private/UnsyncAuth.cpp
new file mode 100644
--- /dev/null
+++ b/Engine/Source/Programs/Unsync/Private/UnsyncAuth.cpp
@@ -0,0 +1,716 @@
+// Copyright Epic Games, Inc. All Rights Reserved.
+
+#include "UnsyncAuth.h"
+#include "UnsyncFile.h"
+#include "UnsyncHttp.h"
+#include "UnsyncLog.h"
+#include "UnsyncProxy.h"
+
+#include <fmt/format.h>
+#include <json11.hpp>
+
+#if UNSYNC_USE_TLS
+
+#	include <openssl/err.h>
+#	include <openssl/rand.h>
+#	include <openssl/sha.h>
+#	include <openssl/evp.h>  // Base64 encoding
+
+#endif	// UNSYNC_USE_TLS
+
+#if UNSYNC_USE_TLS
+namespace unsync {
+
+std::string
+SecureRandomBytesAsHexString(uint32 NumBytes)
+{
+	static constexpr int NumStackBytes				 = 64;
+	unsigned char		 StackStorage[NumStackBytes] = {};
+	FBuffer				 DynamicStorage;
+
+	unsigned char* RandomState = nullptr;
+
+	if (NumBytes <= NumStackBytes)
+	{
+		RandomState = StackStorage;
+	}
+	else
+	{
+		DynamicStorage.Resize(NumBytes);
+		RandomState = DynamicStorage.Data();
+	}
+
+	int RandResult = RAND_bytes(RandomState, NumBytes);
+	if (RandResult != 1)
+	{
+		int ErrorCode = ERR_get_error();
+		UNSYNC_FATAL(L"Failed to generate secure random number. Error code: %d", ErrorCode);
+	}
+
+	return BytesToHexString(RandomState, NumBytes);
+}
+
+FHash256
+HashSha256Bytes(const uint8* Data, uint64 Size)
+{
+	FHash256 Result = {};
+	static_assert(sizeof(Result.Data) == SHA256_DIGEST_LENGTH, "Unexpected SHA256 output buffer size");
+
+	SHA256_CTX ShaCtx = {};
+
+	UNSYNC_ASSERT(SHA256_Init(&ShaCtx) == 1);
+	UNSYNC_ASSERT(SHA256_Update(&ShaCtx, Data, Size) == 1);
+	UNSYNC_ASSERT(SHA256_Final(Result.Data, &ShaCtx) == 1);
+
+	return Result;
+}
+
+std::string
+EncodeBase64(const uint8* Data, uint64 Size)
+{
+	UNSYNC_ASSERT(Size <= std::numeric_limits<int32>::max());
+
+	std::string Result;
+
+	const uint64 ExpectedResultLength = ((Size + 2) / 3) * 4;
+
+	Result.resize(ExpectedResultLength);
+
+	int NumEncodedBytes = EVP_EncodeBlock((unsigned char*)Result.data(), (const unsigned char*)Data, (int)Size);
+	UNSYNC_ASSERT(NumEncodedBytes == ExpectedResultLength)
+
+	return Result;
+}
+
+bool
+DecodeBase64(std::string_view Base64Data, FBuffer& Output)
+{
+	const uint64 ExpectedResultLength = 3 * Base64Data.length() / 4;
+	UNSYNC_ASSERT(ExpectedResultLength <= std::numeric_limits<int32>::max());
+
+	Output.Resize(ExpectedResultLength);  // Conservative size, since EVP_DecodeBlock fills padding with 0
+
+	int NumDecodedBytes = EVP_DecodeBlock((unsigned char*)Output.Data(), (const unsigned char*)Base64Data.data(), (int)Base64Data.length());
+
+	return NumDecodedBytes == ExpectedResultLength;
+}
+
+void
+TransformBase64VanillaToUrlSafe(std::string& Data)
+{
+	std::replace(Data.begin(), Data.end(), '+', '-');
+	std::replace(Data.begin(), Data.end(), '/', '_');
+
+	while (Data.ends_with('='))
+	{
+		Data.pop_back();
+	}
+}
+
+void
+TransformBase64UrlSafeToVanilla(std::string& Data)
+{
+	std::replace(Data.begin(), Data.end(), '-', '+');
+	std::replace(Data.begin(), Data.end(), '_', '/');
+}
+
+std::string
+GetPKCECodeChallenge(std::string_view CodeVerifier)
+{
+	FHash256	CodeVerifierHash = HashSha256Bytes((const uint8*)CodeVerifier.data(), CodeVerifier.size());
+	std::string Result			 = EncodeBase64(CodeVerifierHash.Data, CodeVerifierHash.Size());
+
+	TransformBase64VanillaToUrlSafe(Result);
+
+	return Result;
+}
+
+static const char HttpCallbackResponseOk[] = R"(HTTP/1.1 200 OK
+
+<!DOCTYPE html>
+<html>
+<body>
+<center>
+<h1 style="background-color:#75dd55">Success!</h1>
+<p>Unsync is now authorized. You may close this page.</p>
+</center>
+</body>
+</html>
+)";
+
+static const char HttpCallbackResponseError[] = R"(HTTP/1.1 400 Bad Request
+
+<!DOCTYPE html>
+<html>
+<body>
+<center>
+<h1 style="background-color:#dd5555">Authorization failed!</h1>
+<p>See unsync logs for details. You may close this page.</p>
+</center>
+</body>
+</html>
+)";
+
+struct FHttpCallbackData
+{
+	std::string AuthCode;
+	std::string State;
+};
+
+std::thread
+StartHttpCallbackServer(FSocketHandle	   CallbackListenSocket,
+						std::string_view   ExpectedPath,
+						std::string_view   RandomState,
+						FHttpCallbackData& HttpCallbackData)
+{
+	return std::thread([CallbackListenSocket, ExpectedPath, RandomState, &HttpCallbackData]() {
+		FSocketHandle CallbackSocket = SocketAccept(CallbackListenSocket);
+
+		static const size_t MaxRecvSize = 65536;
+		char				RecvBuffer[MaxRecvSize];
+
+		int32 ReceivedBytes = SocketRecvAny(CallbackSocket, RecvBuffer, MaxRecvSize);
+
+		UNSYNC_VERBOSE2(L"HTTP Callback:\n%.*hs", ReceivedBytes, RecvBuffer);
+
+		std::string_view RequestStr(RecvBuffer, ReceivedBytes);
+
+		std::string ExpectedCallbackPrefix = std::format("GET /{}", ExpectedPath);
+
+		if (RequestStr.starts_with(ExpectedCallbackPrefix))
+		{
+			// Trim request string, removing HTTP headers
+			{
+				size_t RequestEndPos = RequestStr.find("\n");
+				if (RequestEndPos != std::string::npos)
+				{
+					RequestStr = RequestStr.substr(0, RequestEndPos);
+				}
+			}
+
+			auto ExtractValue = [](std::string_view RequestStr, std::string_view Key) -> std::string_view {
+				size_t			 Pos	= RequestStr.find(Key);
+				std::string_view Result = {};
+				if (Pos != std::string::npos)
+				{
+					Result = RequestStr.substr(Pos + Key.length());
+					Result = Result.substr(0, Result.find_first_of("& \n"));
+				}
+				return Result;
+			};
+
+			HttpCallbackData.AuthCode = ExtractValue(RequestStr, "code=");
+			HttpCallbackData.State	  = ExtractValue(RequestStr, "state=");
+
+			if (HttpCallbackData.State == RandomState && !HttpCallbackData.AuthCode.empty())
+			{
+				SocketSend(CallbackSocket, HttpCallbackResponseOk, strlen(HttpCallbackResponseOk));
+			}
+			else
+			{
+				// TODO: could report more detailed error to the browser, but probably just the log file is sufficient
+				SocketSend(CallbackSocket, HttpCallbackResponseError, strlen(HttpCallbackResponseError));
+			}
+		}
+		else
+		{
+			const char ResponseNotFound[] = "HTTP/1.1 404 Not Found";
+			SocketSend(CallbackSocket, ResponseNotFound, strlen(ResponseNotFound));
+		}
+
+		SocketClose(CallbackSocket);
+	});
+};
+
+TResult<FAuthToken>
+AcquireAuthToken(const FAuthDesc& AuthDesc)
+{
+	if (AuthDesc.AuthorizationEndpoint.empty())
+	{
+		return AppError(L"Authorization endpoint is required");
+	}
+
+	if (AuthDesc.TokenEndpoint.empty())
+	{
+		return AppError(L"Token endpoint is required");
+	}
+
+	if (AuthDesc.Callback.empty())
+	{
+		return AppError(L"Callback URI is required");
+	}
+
+	TResult<FRemoteDesc> CallbackServerDescResult = FRemoteDesc::FromUrl(AuthDesc.Callback);
+	if (CallbackServerDescResult.IsError())
+	{
+		return AppError(L"Failed to parse callback URI");
+	}
+
+	const FRemoteDesc& CallbackServerDesc = CallbackServerDescResult.GetData();
+
+	FAuthToken Result;
+
+	FHttpConnection AuthServerConnection = FHttpConnection::CreateDefaultHttps(AuthDesc.ServerHost);
+
+	const uint16 CallbackPortNumber = CallbackServerDesc.HostPort;
+
+	FSocketHandle CallbackListenSocket = SocketListenTcp("127.0.0.1", CallbackPortNumber);
+
+	std::string RandomState	  = SecureRandomBytesAsHexString(16);
+	std::string NonceStr	  = SecureRandomBytesAsHexString(16);
+	std::string CodeVerifier  = SecureRandomBytesAsHexString(64);
+	std::string CodeChallenge = GetPKCECodeChallenge(CodeVerifier);
+	std::string CallbackUrl	  = AuthDesc.Callback;
+
+	std::string AudienceParam;
+	if (!AuthDesc.Audience.empty())
+	{
+		AudienceParam = fmt::format("audience={}&", AuthDesc.Audience);
+	}
+
+	std::string AuthorizeUrl = fmt::format(
+		"https://{}{}?"
+		"client_id={}&"
+		"{}"  // optional audience parameter
+		"response_type=code&"
+		"scope=offline_access%20profile%20openid%20email&"
+		"code_challenge_method=S256&"
+		"code_challenge={}&"
+		"state={}&"
+		"redirect_uri={}",
+		AuthDesc.ServerHost,
+		AuthDesc.AuthorizationEndpoint,
+		AuthDesc.ClientId,
+		AudienceParam,
+		CodeChallenge,
+		RandomState,
+		CallbackUrl);
+
+	FHttpCallbackData HttpCallbackData;
+
+	std::thread ServerThread = StartHttpCallbackServer(CallbackListenSocket, CallbackServerDesc.RequestPath, RandomState, HttpCallbackData);
+
+	OpenUrlInDefaultBrowser(AuthorizeUrl.c_str());
+
+	UNSYNC_LOG(L"Waiting for HTTP callback on port %d...", int(CallbackPortNumber));
+
+	ServerThread.join();
+	SocketClose(CallbackListenSocket);
+
+	if (RandomState != HttpCallbackData.State)
+	{
+		return AppError(L"Callback state value mismatch");
+	}
+
+	if (HttpCallbackData.AuthCode.empty())
+	{
+		return AppError(L"Did not receive authorization code callback");
+	}
+
+	std::string AccessToken;
+	std::string RefreshToken;
+	std::string IdToken;
+	std::string TokenType;
+	int64		ExpiresInSeconds = 0;
+
+	// TODO: only try to acquire new token if close to expiry
+
+	// Use authorization code to acquire tokens
+	{
+		std::string TokenPayload = fmt::format(
+			"grant_type=authorization_code&"
+			"client_id={}&"
+			"code={}&"
+			"code_verifier={}&"
+			"redirect_uri={}",
+			AuthDesc.ClientId,
+			HttpCallbackData.AuthCode,
+			CodeVerifier,
+			CallbackUrl);
+
+		FHttpRequest Request;
+		Request.Url				   = AuthDesc.TokenEndpoint;
+		Request.Method			   = EHttpMethod::POST;
+		Request.PayloadContentType = EHttpContentType::Application_WWWFormUrlEncoded;
+		Request.Payload			   = FBufferView{(const uint8*)TokenPayload.data(), (uint64)TokenPayload.size()};
+
+		FHttpResponse Response = HttpRequest(AuthServerConnection, Request);
+
+		if (Response.Success())
+		{
+			using namespace json11;
+			std::string JsonString = std::string(Response.AsStringView());
+
+			std::string JsonErrorString;
+			Json		JsonObject = Json::parse(JsonString, JsonErrorString);
+
+			if (!JsonErrorString.empty())
+			{
+				return AppError(fmt::format("JSON error while parsing token: {}", JsonErrorString.c_str()));
+			}
+
+			AccessToken		 = JsonObject["access_token"].string_value();
+			RefreshToken	 = JsonObject["refresh_token"].string_value();
+			IdToken			 = JsonObject["id_token"].string_value();
+			TokenType		 = JsonObject["token_type"].string_value();
+			ExpiresInSeconds = int64(JsonObject["expires_in"].number_value());
+
+			Result.Raw = JsonString;
+		}
+		else
+		{
+			return HttpError(L"Could not acquire authorization code", Response.Code);
+		}
+	}
+
+	if (AccessToken.empty())
+	{
+		return AppError(L"Did not receive new access token");
+	}
+
+	Result.Access  = std::move(AccessToken);
+	Result.Refresh = std::move(RefreshToken);
+
+	return ResultOk(std::move(Result));
+}
+
+TResult<FAuthUserInfo>
+GetUserInfo(FHttpConnection& HttpConnection, const FAuthDesc& AuthDesc, const FAuthToken& AuthToken)
+{
+	std::string AuthHeader = fmt::format("Authorization: Bearer {}", AuthToken.Access);
+
+	if (AuthDesc.UserInfoEndpoint.empty())
+	{
+		return AppError(L"User info endpoint is unknown");
+	}
+
+	FHttpRequest Request;
+	Request.Url			  = AuthDesc.UserInfoEndpoint;
+	Request.Method		  = EHttpMethod::GET;
+	Request.CustomHeaders = AuthHeader;
+
+	FHttpResponse Response = HttpRequest(HttpConnection, Request);
+
+	if (Response.Success())
+	{
+		using namespace json11;
+
+		std::string JsonString = std::string(Response.AsStringView());
+
+		std::string JsonErrorString;
+		Json		JsonObject = Json::parse(JsonString, JsonErrorString);
+
+		if (!JsonErrorString.empty())
+		{
+			return AppError(fmt::format("JSON error while parsing user info: {}", JsonErrorString.c_str()));
+		}
+
+		FAuthUserInfo Result;
+
+		Result.Sub		  = JsonObject["sub"].string_value();
+		Result.Name		  = JsonObject["name"].string_value();
+		Result.Nickname	  = JsonObject["nickname"].string_value();
+		Result.GivenName  = JsonObject["given_name"].string_value();
+		Result.FamilyName = JsonObject["family_name"].string_value();
+		Result.Email	  = JsonObject["email"].string_value();
+
+		return ResultOk(Result);
+	}
+	else
+	{
+		return HttpError(L"Could not query user info from authorization server", Response.Code);
+	}
+}
+
+TResult<FAuthToken>
+RefreshAuthToken(const FAuthDesc& AuthDesc, const FAuthToken& PreviousToken)
+{
+	if (AuthDesc.TokenEndpoint.empty())
+	{
+		return AppError(L"Token endpoint is unknown");
+	}
+
+	FAuthToken Result = PreviousToken;
+
+	FHttpConnection AuthServerConnection = FHttpConnection::CreateDefaultHttps(AuthDesc.ServerHost);
+
+	std::string AccessToken;
+	std::string RefreshToken;
+	std::string IdToken;
+	std::string TokenType;
+	int64		ExpiresInSeconds = 0;
+
+	// Use refresh token to acquire new tokens
+	{
+		std::string TokenPayload = fmt::format(
+			"grant_type=refresh_token&"
+			"client_id={}&"
+			"refresh_token={}",
+			AuthDesc.ClientId,
+			PreviousToken.Refresh);
+
+		FHttpRequest Request;
+		Request.Url				   = AuthDesc.TokenEndpoint;
+		Request.Method			   = EHttpMethod::POST;
+		Request.PayloadContentType = EHttpContentType::Application_WWWFormUrlEncoded;
+		Request.Payload			   = FBufferView{(const uint8*)TokenPayload.data(), (uint64)TokenPayload.size()};
+
+		FHttpResponse Response = HttpRequest(AuthServerConnection, Request);
+
+		if (Response.Success())
+		{
+			using namespace json11;
+			std::string JsonString = std::string(Response.AsStringView());
+
+			std::string JsonErrorString;
+			Json		JsonObject = Json::parse(JsonString, JsonErrorString);
+
+			if (!JsonErrorString.empty())
+			{
+				return AppError(fmt::format("JSON error while parsing token: {}", JsonErrorString.c_str()));
+			}
+
+			AccessToken		 = JsonObject["access_token"].string_value();
+			RefreshToken	 = JsonObject["refresh_token"].string_value();
+			IdToken			 = JsonObject["id_token"].string_value();
+			TokenType		 = JsonObject["token_type"].string_value();
+			ExpiresInSeconds = int64(JsonObject["expires_in"].number_value());
+
+			Result.Raw = JsonString;
+		}
+		else
+		{
+			return HttpError(L"Could not acquire authorization code", Response.Code);
+		}
+	}
+
+	if (AccessToken.empty())
+	{
+		return AppError(L"Did not receive new access token");
+	}
+
+	Result.Access = AccessToken;
+
+	if (!RefreshToken.empty())
+	{
+		Result.Refresh = RefreshToken;
+	}
+
+	return ResultOk(Result);
+}
+
+std::string
+GenerateTokenId(const FRemoteDesc& RemoteDesc)
+{
+	FHash128 Hash = HashBlake3String<FHash128>(RemoteDesc.HostAddress);
+	return HashToHexString(Hash);
+}
+
+bool
+SaveRefreshToken(const FPath& Path, const FAuthToken& AuthToken)
+{
+	return WriteBufferToFile(Path, (const uint8*)AuthToken.Refresh.data(), AuthToken.Refresh.length());
+}
+
+TResult<FAuthToken>
+LoadRefreshToken(const FPath& Path)
+{
+	FBuffer FileBuffer = ReadFileToBuffer(Path);
+	if (FileBuffer.Size())
+	{
+		FAuthToken AuthToken;
+		AuthToken.Refresh.append((const char*)FileBuffer.Data(), FileBuffer.Size());
+		return ResultOk(AuthToken);
+	}
+	else
+	{
+		return AppError(L"Failed to load refresh token from file");
+	}
+}
+
+TResult<FAuthToken>
+RefreshOrAcquireToken(const FAuthDesc& AuthDesc, const FAuthToken& PreviousToken)
+{
+	if (!PreviousToken.Refresh.empty())
+	{
+		UNSYNC_VERBOSE(L"Refreshing access token");
+		TResult<FAuthToken> RefreshResult = RefreshAuthToken(AuthDesc, PreviousToken);
+		if (RefreshResult.IsOk())
+		{
+			return RefreshResult;
+		}
+	}
+
+	UNSYNC_VERBOSE(L"Requesting new access token");
+	return AcquireAuthToken(AuthDesc);
+}
+
+TResult<FAuthToken>
+Authenticate(const FRemoteDesc& RemoteDesc, const FAuthDesc& AuthDesc)
+{
+	FPath UserHomePath = GetUserHomeDirectory();
+
+	std::string TokenId = GenerateTokenId(RemoteDesc);
+
+	FPath UnsyncSettingsPath = UserHomePath / FPath(".unsync");
+	FPath TokenCachePath	 = UnsyncSettingsPath / FPath(TokenId);
+
+	FAuthToken PreviousToken;
+
+	if (PathExists(TokenCachePath))
+	{
+		TResult<FAuthToken> LoadResult = LoadRefreshToken(TokenCachePath);
+		if (FAuthToken* LoadedToken = LoadResult.TryData())
+		{
+			UNSYNC_VERBOSE(L"Loaded cached authentication data");
+			PreviousToken = std::move(*LoadedToken);
+		}
+	}
+
+	TResult<FAuthToken> FreshTokenResult = RefreshOrAcquireToken(AuthDesc, PreviousToken);
+	if (FreshTokenResult.IsError())
+	{
+		return FreshTokenResult;
+	}
+
+	if (!UserHomePath.empty())
+	{
+		CreateDirectories(UnsyncSettingsPath);
+
+		// Allow saving tokens during dry run
+		// TODO: need a dedicated file flag to allow writes during dry run
+
+		const bool bPrevDryRun = GDryRun;
+		GDryRun				   = false;
+
+		bool bSaved = SaveRefreshToken(TokenCachePath, FreshTokenResult.GetData());
+
+		GDryRun = bPrevDryRun;
+
+		if (bSaved)
+		{
+			UNSYNC_VERBOSE2(L"Saved refresh token to file: %s", TokenCachePath.wstring().c_str());
+		}
+	}
+
+	return FreshTokenResult;
+}
+
+bool
+TryAddAuthentication(FRemoteDesc& InOutRemoteDesc)
+{
+	TResult<FAuthDesc> AuthDescResult = GetAuthenticationDesc(InOutRemoteDesc);
+	if (AuthDescResult.IsError())
+	{
+		return false;
+	}
+
+	const FAuthDesc& AuthDesc = AuthDescResult.GetData();
+
+	TResult<FAuthToken> AuthTokenResult = Authenticate(InOutRemoteDesc, AuthDesc);
+	if (AuthTokenResult.IsError())
+	{
+		return false;
+	}
+
+	const std::string& AccessToken = AuthTokenResult->Access;
+
+	// Authentication requires encrypted connection
+	InOutRemoteDesc.bTlsEnable	   = true;
+	InOutRemoteDesc.Authentication = std::make_shared<FBuffer>();
+	InOutRemoteDesc.Authentication->Append((const uint8*)AccessToken.data(), AccessToken.length());
+
+	return true;
+}
+
+TResult<FAuthDesc>
+GetAuthenticationDesc(const FRemoteDesc& RemoteDesc)
+{
+	// TODO: possibly other backend could use automatic authentication also
+	if (RemoteDesc.Protocol != EProtocolFlavor::Unsync)
+	{
+		return AppError(L"Authentication is only implemented for UNSYNC protocol");
+	}
+
+	TResult<ProxyQuery::FHelloResponse> HelloResponseResult = ProxyQuery::Hello(RemoteDesc);
+	if (HelloResponseResult.IsError())
+	{
+		return MoveError<FAuthDesc>(HelloResponseResult);
+	}
+
+	if (!HelloResponseResult->SupportsAuthentication())
+	{
+		return AppError(L"Server does not support authentication");
+	}
+
+	FAuthDesc AuthDesc;
+	AuthDesc.ClientId = HelloResponseResult->AuthClientId;
+	AuthDesc.Audience = HelloResponseResult->AuthAudience;
+	AuthDesc.Callback = HelloResponseResult->CallbackUri;
+	if (AuthDesc.Callback.empty())
+	{
+		AuthDesc.Callback = "http://localhost:8080";  // sensible default
+	}
+
+	TResult<FRemoteDesc> AuthServerDescResult = FRemoteDesc::FromUrl(HelloResponseResult->AuthServerUri);
+	if (AuthServerDescResult.IsError())
+	{
+		return MoveError<FAuthDesc>(AuthServerDescResult);
+	}
+
+	const FRemoteDesc& AuthServerDesc = AuthServerDescResult.GetData();
+
+	AuthDesc.ServerHost = AuthServerDesc.HostAddress;
+
+	std::string ServerApiPrefix;
+	if (!AuthServerDesc.RequestPath.empty())
+	{
+		ServerApiPrefix = fmt::format("/{}", AuthServerDesc.RequestPath);
+	}
+
+	if (!AuthDesc.IsValid())
+	{
+		return AppError(L"Mandatory authentication parameters not found");
+	}
+
+	FHttpConnection AuthServerConnection = FHttpConnection::CreateDefaultHttps(AuthDesc.ServerHost);
+	std::string		ConfigEndpoint		 = fmt::format("{}/.well-known/openid-configuration", ServerApiPrefix);
+
+	FHttpResponse ConfigResponse = HttpRequest(AuthServerConnection, EHttpMethod::GET, ConfigEndpoint);
+	if (ConfigResponse.Success())
+	{
+		using namespace json11;
+		std::string JsonString = std::string(ConfigResponse.AsStringView());
+
+		std::string JsonErrorString;
+		Json		JsonObject = Json::parse(JsonString, JsonErrorString);
+
+		std::string EndpointPrefix = fmt::format("https://{}", AuthDesc.ServerHost);
+
+		if (JsonErrorString.empty())
+		{
+			auto ExtractEndpoint = [&EndpointPrefix, &JsonObject](const char* FieldName) -> std::string {
+				if (auto& Field = JsonObject[FieldName]; Field.is_string())
+				{
+					const std::string& Value = Field.string_value();
+					if (Value.starts_with(EndpointPrefix))
+					{
+						return Value.substr(EndpointPrefix.length());
+					}
+				}
+				return {};
+			};
+
+			AuthDesc.AuthorizationEndpoint = ExtractEndpoint("authorization_endpoint");
+			AuthDesc.TokenEndpoint		   = ExtractEndpoint("token_endpoint");
+			AuthDesc.UserInfoEndpoint	   = ExtractEndpoint("userinfo_endpoint");
+			AuthDesc.JwksUri			   = JsonObject["jwks_uri"].string_value();
+		}
+	}
+
+	return ResultOk(AuthDesc);
+}
+
+}  // namespace unsync
+
+#endif	// UNSYNC_USE_TLS
diff --git a/Engine/Source/Programs/Unsync/Private/UnsyncAuth.h b/Engine/Source/Programs/Unsync/Private/UnsyncAuth.h
new file mode 100644
--- /dev/null
+++ b/Engine/Source/Programs/Unsync/Private/UnsyncAuth.h
@@ -0,0 +1,98 @@
+// Copyright Epic Games, Inc. All Rights Reserved.
+
+#pragma once
+
+#include "UnsyncBuffer.h"
+#include "UnsyncCommon.h"
+#include "UnsyncError.h"
+#include "UnsyncHash.h"
+
+#include <string>
+
+#if UNSYNC_USE_TLS
+
+namespace unsync {
+
+struct FHttpConnection;
+struct FRemoteDesc;
+
+struct FAuthToken
+{
+	// Raw data acquired from authentication endpoint
+	std::string Raw;
+
+	// Parsed data for convenience
+	std::string Access;
+	std::string Refresh;
+	int64		ExirationTime = 0; // UNIX timestamp in seconds (0 if unknown)
+};
+
+struct FAuthDesc
+{
+	// Mandatory configuration
+	std::string ServerHost; // only the hostname itself
+	std::string ClientId;
+	
+	// Mandatory fields for certain APIs
+	std::string AuthorizationEndpoint;
+	std::string TokenEndpoint;
+	std::string UserInfoEndpoint;
+	std::string JwksUri;
+
+	// Optional configuration
+	std::string Audience;
+	std::string Callback;	// OAuth flow redirection URL
+
+	// Check that mandatory fields have values
+	bool IsValid() const
+	{
+		return !ServerHost.empty() && !ClientId.empty();
+	}
+};
+
+// Query OIDC/OAuth2 configuration using server handshake
+TResult<FAuthDesc> GetAuthenticationDesc(const FRemoteDesc& RemoteDesc);
+
+// Perform the complete authentication flow:
+// - Attempt use stored refresh token first
+// - If refresh is not possible, use PKCE Authentication flow to get new tokens
+// - Save refresh token in user directory for future use
+TResult<FAuthToken> Authenticate(const FRemoteDesc& RemoteDesc, const FAuthDesc& AuthDesc);
+
+// Auth utility functions
+
+TResult<FAuthToken> AcquireAuthToken(const FAuthDesc& AuthDesc);
+TResult<FAuthToken> RefreshAuthToken(const FAuthDesc& AuthDesc, const FAuthToken& PreviousToken);
+TResult<FAuthToken> RefreshOrAcquireToken(const FAuthDesc& AuthDesc, const FAuthToken& PreviousToken);
+
+struct FAuthUserInfo
+{
+	std::string Sub;
+	std::string Name;
+	std::string Nickname;
+	std::string GivenName;
+	std::string FamilyName;
+	std::string Email;
+};
+TResult<FAuthUserInfo> GetUserInfo(FHttpConnection& HttpConnection, const FAuthDesc& AuthDesc, const FAuthToken& AuthToken);
+
+std::string GenerateTokenId(const FRemoteDesc& RemoteDesc);
+
+bool				SaveRefreshToken(const FPath& Path, const FAuthToken& AuthToken);
+TResult<FAuthToken> LoadRefreshToken(const FPath& Path);
+
+std::string SecureRandomBytesAsHexString(uint32 NumBytes);
+FHash256	HashSha256Bytes(const uint8* Data, uint64 Size);
+std::string EncodeBase64(const uint8* Data, uint64 Size);
+bool		DecodeBase64(std::string_view Base64Data, FBuffer& Output);
+
+void TransformBase64VanillaToUrlSafe(std::string& Base64Vanilla);
+void TransformBase64UrlSafeToVanilla(std::string& Base64UrlSafe);
+
+std::string GetPKCECodeChallenge(std::string_view CodeVerifier);
+
+bool TryAddAuthentication(FRemoteDesc& InOutRemoteDesc);
+
+}  // namespace unsync
+
+#endif	// UNSYNC_USE_TLS
diff --git a/Engine/Source/Programs/Unsync/Private/UnsyncCmdQuery.cpp b/Engine/Source/Programs/Unsync/Private/UnsyncCmdQuery.cpp
--- a/Engine/Source/Programs/Unsync/Private/UnsyncCmdQuery.cpp
+++ b/Engine/Source/Programs/Unsync/Private/UnsyncCmdQuery.cpp
@@ -1,7 +1,9 @@
 // Copyright Epic Games, Inc. All Rights Reserved.
 
 #include "UnsyncCmdQuery.h"
+#include "UnsyncFile.h"
 #include "UnsyncHttp.h"
+#include "UnsyncProxy.h"
 #include "UnsyncThread.h"
 #include "UnsyncUtil.h"
 
@@ -9,6 +11,12 @@
 #include <algorithm>
 #include <json11.hpp>
 
+#include <fmt/format.h>
+
+#if UNSYNC_USE_TLS
+#	include "UnsyncAuth.h"
+#endif	// UNSYNC_USE_TLS
+
 namespace unsync {
 
 using FMirrorInfoResult = TResult<std::vector<FMirrorInfo>>;
@@ -94,49 +102,104 @@ RunQueryMirrors(const FRemoteDesc& RemoteDesc)
 }
 
 int32
-CmdQuery(const FCmdQueryOptions& Options)
+CmdQueryMirrors(const FCmdQueryOptions& Options)
 {
-	if (Options.Query == "mirrors")
+	FMirrorInfoResult MirrorsResult = RunQueryMirrors(Options.Remote);
+	if (MirrorsResult.IsError())
 	{
-		FMirrorInfoResult MirrorsResult = RunQueryMirrors(Options.Remote);
-		if (MirrorsResult.IsError())
-		{
-			LogError(MirrorsResult.GetError());
-			return 1;
-		}
+		LogError(MirrorsResult.GetError());
+		return 1;
+	}
 
-		std::vector<FMirrorInfo> Mirrors = MirrorsResult.GetData();
+	std::vector<FMirrorInfo> Mirrors = MirrorsResult.GetData();
 
-		ParallelForEach(Mirrors.begin(), Mirrors.end(), [](FMirrorInfo& Mirror) {
-			Mirror.Ping = RunHttpPing(Mirror.Address, Mirror.Port);
-		});
+	ParallelForEach(Mirrors.begin(), Mirrors.end(), [](FMirrorInfo& Mirror) { Mirror.Ping = RunHttpPing(Mirror.Address, Mirror.Port); });
 
-		std::sort(Mirrors.begin(), Mirrors.end(), [](const FMirrorInfo& InA, const FMirrorInfo& InB) {
-			double A = InA.Ping > 0 ? InA.Ping : FLT_MAX;
-			double B = InB.Ping > 0 ? InB.Ping : FLT_MAX;
-			return A < B;
-		});
+	std::sort(Mirrors.begin(), Mirrors.end(), [](const FMirrorInfo& InA, const FMirrorInfo& InB) {
+		double A = InA.Ping > 0 ? InA.Ping : FLT_MAX;
+		double B = InB.Ping > 0 ? InB.Ping : FLT_MAX;
+		return A < B;
+	});
 
-		LogPrintf(ELogLevel::Info, L"[\n");
+	LogPrintf(ELogLevel::Info, L"[\n");
 
-		for (size_t I = 0; I < Mirrors.size(); ++I)
+	for (size_t I = 0; I < Mirrors.size(); ++I)
+	{
+		const FMirrorInfo& Mirror = Mirrors[I];
+
+		LogPrintf(ELogLevel::Info,
+				  L"  {\"address\":\"%hs\", \"port\":%d, \"ok\":%hs, \"ping\":%d, \"name\":\"%hs\"}%hs\n",
+				  Mirror.Address.c_str(),
+				  Mirror.Port,
+				  Mirror.Ping > 0 ? "true" : "false",
+				  int32(Mirror.Ping * 1000.0),
+				  Mirror.Name.c_str(),
+				  I + 1 == Mirrors.size() ? "" : ",");
+	}
+
+	LogPrintf(ELogLevel::Info, L"]\n");
+
+	return 0;
+}
+
+#if UNSYNC_USE_TLS
+int32
+CmdQueryLogin(const FCmdQueryOptions& Options)
+{
+	TResult<FAuthDesc> AuthDescResult = GetAuthenticationDesc(Options.Remote);
+	if (AuthDescResult.IsError())
+	{
+		LogError(AuthDescResult.GetError());
+		return -1;
+	}
+
+	const FAuthDesc& AuthDesc = AuthDescResult.GetData();
+
+	TResult<FAuthToken> AuthToken = Authenticate(Options.Remote, AuthDesc);
+
+	if (AuthToken.IsOk())
+	{
+		UNSYNC_LOG(L"Login successful");
+
+		FHttpConnection AuthConnection = FHttpConnection::CreateDefaultHttps(AuthDesc.ServerHost);
+
+		TResult<FAuthUserInfo> UserInfoResult = GetUserInfo(AuthConnection, AuthDesc, AuthToken.GetData());
+
+		if (const FAuthUserInfo* UserInfo = UserInfoResult.TryData())
 		{
-			const FMirrorInfo& Mirror = Mirrors[I];
-
-			LogPrintf(ELogLevel::Info,
-					  L"  {\"address\":\"%hs\", \"port\":%d, \"ok\":%hs, \"ping\":%d, \"name\":\"%hs\"}%hs\n",
-					  Mirror.Address.c_str(),
-					  Mirror.Port,
-					  Mirror.Ping > 0 ? "true" : "false",
-					  int32(Mirror.Ping * 1000.0),
-					  Mirror.Name.c_str(),
-					  I + 1 == Mirrors.size() ? "" : ",");
+			if (UserInfo->Email.length() && UserInfo->Name.length())
+			{
+				UNSYNC_VERBOSE(L"Authenticated user: %hs (%hs)", UserInfo->Name.c_str(), UserInfo->Email.c_str());
+			}
+			else if (UserInfo->Name.length())
+			{
+				UNSYNC_VERBOSE(L"Authenticated user: %hs", UserInfo->Name.c_str());
+			}
 		}
 
-		LogPrintf(ELogLevel::Info, L"]\n");
-
 		return 0;
 	}
+	else
+	{
+		LogError(AuthToken.GetError());
+		return -1;
+	}
+}
+#endif	// UNSYNC_USE_TLS
+
+int32
+CmdQuery(const FCmdQueryOptions& Options)
+{
+	if (Options.Query == "mirrors")
+	{
+		return CmdQueryMirrors(Options);
+	}
+#if UNSYNC_USE_TLS
+	else if (Options.Query == "login")
+	{
+		return CmdQueryLogin(Options);
+	}
+#endif	// UNSYNC_USE_TLS
 	else
 	{
 		UNSYNC_ERROR(L"Unknown query command");
diff --git a/Engine/Source/Programs/Unsync/Private/UnsyncCommon.h b/Engine/Source/Programs/Unsync/Private/UnsyncCommon.h
--- a/Engine/Source/Programs/Unsync/Private/UnsyncCommon.h
+++ b/Engine/Source/Programs/Unsync/Private/UnsyncCommon.h
@@ -22,6 +22,10 @@ using int64 = int64_t;
 
 using FPath = std::filesystem::path;
 
+#ifndef UNSYNC_USE_TLS
+#	define UNSYNC_USE_TLS 0
+#endif
+
 #define UNSYNC_DISALLOW_COPY_ASSIGN(T) \
 	T(const T&) = delete;              \
 	T& operator=(const T&) = delete;
diff --git a/Engine/Source/Programs/Unsync/Private/UnsyncError.h b/Engine/Source/Programs/Unsync/Private/UnsyncError.h
--- a/Engine/Source/Programs/Unsync/Private/UnsyncError.h
+++ b/Engine/Source/Programs/Unsync/Private/UnsyncError.h
@@ -128,6 +128,13 @@ private:
 	std::variant<T, Et> DataOrError;
 };
 
+template<typename T, typename Te, typename T2>
+inline TResult<T, Te>
+MoveError(TResult<T2, Te>& X)
+{
+	return TResult<T, Te>(std::move(X.GetError()));
+}
+
 template<typename Tx, typename E = FError>
 static TResult<Tx, E>
 ResultOk(Tx Ok)
diff --git a/Engine/Source/Programs/Unsync/Private/UnsyncHttp.cpp b/Engine/Source/Programs/Unsync/Private/UnsyncHttp.cpp
--- a/Engine/Source/Programs/Unsync/Private/UnsyncHttp.cpp
+++ b/Engine/Source/Programs/Unsync/Private/UnsyncHttp.cpp
@@ -128,6 +128,10 @@ struct FHttpParser
 			{
 				Response.ContentType = EHttpContentType::Application_UECB;
 			}
+			else if (MatchUncased(PendingValue, "application/x-www-form-urlencoded"))
+			{
+				Response.ContentType = EHttpContentType::Application_WWWFormUrlEncoded;
+			}
 			else if (MatchUncased(PendingValue, "text/html"))
 			{
 				Response.ContentType = EHttpContentType::Text_Html;
@@ -327,6 +331,9 @@ HttpRequestBegin(FHttpConnection& Connection, const FHttpRequest& Request)
 			case EHttpContentType::Application_UECB:
 				HttpHeader += "Content-type: application/x-ue-cb\r\n";
 				break;
+			case EHttpContentType::Application_WWWFormUrlEncoded:
+				HttpHeader += "Content-type: application/x-www-form-urlencoded\r\n";
+				break;
 			case EHttpContentType::Application_Json:
 				HttpHeader += "Content-type: application/json\r\n";
 				break;
@@ -472,6 +479,21 @@ FHttpConnection::FHttpConnection(const FHttpConnection& Other)
 {
 }
 
+FHttpConnection
+FHttpConnection::CreateDefaultHttp(const std::string_view InHostAddress, uint16 Port)
+{
+	return FHttpConnection(InHostAddress, Port, nullptr);
+}
+
+FHttpConnection
+FHttpConnection::CreateDefaultHttps(const std::string_view InHostAddress, uint16 Port)
+{
+	FTlsClientSettings TlsSettings;
+	TlsSettings.Subject			   = InHostAddress.data();
+	TlsSettings.bVerifyCertificate = true;
+	return FHttpConnection(InHostAddress, Port, &TlsSettings);
+}
+
 bool
 FHttpConnection::Open()
 {
diff --git a/Engine/Source/Programs/Unsync/Private/UnsyncHttp.h b/Engine/Source/Programs/Unsync/Private/UnsyncHttp.h
--- a/Engine/Source/Programs/Unsync/Private/UnsyncHttp.h
+++ b/Engine/Source/Programs/Unsync/Private/UnsyncHttp.h
@@ -22,6 +22,7 @@ enum class EHttpContentType
 	Text_Plain,
 	Application_OctetStream,
 	Application_Json,
+	Application_WWWFormUrlEncoded,
 	Application_UECB,  // Unreal Engine Compact Binary
 };
 
@@ -62,6 +63,9 @@ struct FHttpConnection
 	FHttpConnection(const std::string_view InHostAddress, uint16 InPort, const FTlsClientSettings* TlsSettings = nullptr);
 	FHttpConnection(const FHttpConnection& Other);
 
+	[[nodiscard]] static FHttpConnection CreateDefaultHttp(const std::string_view InHostAddress, uint16 Port = 80);
+	[[nodiscard]] static FHttpConnection CreateDefaultHttps(const std::string_view InHostAddress, uint16 Port = 443);
+
 	bool Open();
 	void Close();
 
diff --git a/Engine/Source/Programs/Unsync/Private/UnsyncMain.cpp b/Engine/Source/Programs/Unsync/Private/UnsyncMain.cpp
--- a/Engine/Source/Programs/Unsync/Private/UnsyncMain.cpp
+++ b/Engine/Source/Programs/Unsync/Private/UnsyncMain.cpp
@@ -14,6 +14,7 @@
 #include "UnsyncTest.h"
 #include "UnsyncThread.h"
 #include "UnsyncUtil.h"
+#include "UnsyncAuth.h"
 
 UNSYNC_THIRD_PARTY_INCLUDES_START
 #if UNSYNC_PLATFORM_WINDOWS
@@ -105,9 +106,9 @@ InnerMain(int Argc, char** Argv)
 #endif	// UNSYNC_USE_TLS
 
 	auto AddProxyOptions = [&RemoteAddressUtf8, &bNoProxySelect](CLI::App* App) {
-		App->add_option("--proxy, --remote",
+		App->add_option("--proxy, --remote, --server",
 						RemoteAddressUtf8,
-						"FProxy server address ([transport://]address[:port][/request][#namespace])");
+						"Download server address ([transport://]address[:port][/request][#namespace])");
 		App->add_flag("--no-proxy-select",
 					  bNoProxySelect,
 					  "Skip automatic server selection and use the exact one specified by command line or environment variable");
@@ -707,6 +708,18 @@ InnerMain(int Argc, char** Argv)
 			ScavengeRoot = FPath{};
 		}
 
+#if UNSYNC_USE_TLS
+		{
+			UNSYNC_VERBOSE(L"Attempting to authenticate");
+			UNSYNC_LOG_INDENT;
+			bool bUsingAuthentication = TryAddAuthentication(RemoteDesc);
+			if (bUsingAuthentication)
+			{
+				UNSYNC_VERBOSE(L"Authentication enabled");
+			}
+		}
+#endif	// UNSYNC_USE_TLS
+
 		FCmdSyncOptions SyncOptions;
 
 		SyncOptions.Algorithm			   = Algorithm;
diff --git a/Engine/Source/Programs/Unsync/Private/UnsyncProtocol.h b/Engine/Source/Programs/Unsync/Private/UnsyncProtocol.h
--- a/Engine/Source/Programs/Unsync/Private/UnsyncProtocol.h
+++ b/Engine/Source/Programs/Unsync/Private/UnsyncProtocol.h
@@ -120,6 +120,7 @@ struct FHandshakePacket
 };
 
 static constexpr uint64 COMMAND_ID_DISCONNECT = 0;
+static constexpr uint64 COMMAND_ID_AUTHENTICATE = 0xAA77CB56ABD7153Aull;
 static constexpr uint64 COMMAND_ID_GET_BLOCKS = 0xBBE2A1CECC8C949Cull;
 
 struct FCommandPacket
@@ -129,6 +130,13 @@ struct FCommandPacket
 	uint64					CommandId = 0;
 };
 
+struct FBufferPacket
+{
+	static constexpr uint64 MAGIC		  = 0x6539A89058A3400Aull;
+	uint64					Magic		  = MAGIC;
+	uint32					DataSizeBytes = 0;
+};
+
 struct FFileListPacket
 {
 	static constexpr uint64 MAGIC		  = 0x28B96050A327172Aull;
diff --git a/Engine/Source/Programs/Unsync/Private/UnsyncProxy.cpp b/Engine/Source/Programs/Unsync/Private/UnsyncProxy.cpp
--- a/Engine/Source/Programs/Unsync/Private/UnsyncProxy.cpp
+++ b/Engine/Source/Programs/Unsync/Private/UnsyncProxy.cpp
@@ -121,6 +121,28 @@ FUnsyncProtocolImpl::FUnsyncProtocolImpl(const FRemoteDesc&				RemoteDesc,
 			return true;
 		}();
 	}
+
+	if (IsValid() && Features.bAuthentication && RemoteDesc.Authentication)
+	{
+		bool bOk = IsValid();
+
+		FCommandPacket Packet;
+		Packet.CommandId = COMMAND_ID_AUTHENTICATE;
+		bOk &= SendStruct(*SocketHandle, Packet);
+		bOk &= SendBuffer(*SocketHandle, *RemoteDesc.Authentication);
+
+		int32 ResultSize = 0;
+		bOk &= SocketRecvT(*SocketHandle, ResultSize);
+
+		FBuffer ResultBuffer;
+		if (ResultSize)
+		{
+			ResultBuffer.Resize(ResultSize);
+			bOk &= (SocketRecvAll(*SocketHandle, ResultBuffer.Data(), ResultSize) == ResultSize);
+		}
+
+		// TODO: parse authentication result packet and report errors
+	}
 }
 
 FProxy::~FProxy()
@@ -396,6 +418,14 @@ TResult<FHelloResponse> Hello(const FRemoteDesc& RemoteDesc)
 		Result.SessionId = Field.string_value();
 	}
 
+	if (auto& Field = JsonObject["auth"]; Field.is_object())
+	{
+		Result.AuthServerUri = Field["server"].string_value();
+		Result.AuthClientId	 = Field["client_id"].string_value();
+		Result.AuthAudience	 = Field["audience"].string_value();
+		Result.CallbackUri	 = Field["callback"].string_value();
+	}
+
 	if (auto& Field = JsonObject["features"]; Field.is_array())
 	{
 		Result.FeatureNames.reserve(Field.array_items().size());
@@ -409,11 +439,14 @@ TResult<FHelloResponse> Hello(const FRemoteDesc& RemoteDesc)
 				{
 					Result.Features.bTelemetry = true;
 				}
-
-				if (Elem.string_value() == "mirrors")
+				else if (Elem.string_value() == "mirrors")
 				{
 					Result.Features.bMirrors = true;
 				}
+				else if (Elem.string_value() == "authentication")
+				{
+					Result.Features.bAuthentication = true;
+				}
 			}
 		}
 	}
diff --git a/Engine/Source/Programs/Unsync/Private/UnsyncProxy.h b/Engine/Source/Programs/Unsync/Private/UnsyncProxy.h
--- a/Engine/Source/Programs/Unsync/Private/UnsyncProxy.h
+++ b/Engine/Source/Programs/Unsync/Private/UnsyncProxy.h
@@ -87,6 +87,7 @@ struct FRemoteProtocolFeatures
 {
 	bool bTelemetry = false;
 	bool bMirrors = false;
+	bool bAuthentication = false;
 };
 
 struct FTelemetryEventSyncComplete
@@ -190,11 +191,27 @@ struct FHelloResponse
 	std::string VersionGit;
 	std::string SessionId;
 
+	std::string AuthServerUri;
+	std::string AuthClientId;
+	std::string AuthAudience;
+	std::string CallbackUri;
+
 	std::vector<std::string> FeatureNames;
 	FRemoteProtocolFeatures	 Features;
+
+	bool SupportsAuthentication() const
+	{
+		return Features.bAuthentication && !AuthServerUri.empty() && !AuthClientId.empty();
+	}
 };
 TResult<FHelloResponse> Hello(const FRemoteDesc& RemoteDesc);
 
+struct FLoginResponse
+{
+	std::shared_ptr<FBuffer> Authentication;
+};
+TResult<FLoginResponse> QueryLogin(const FRemoteDesc& RemoteDesc);
+
 } 
 
 }  // namespace unsync
diff --git a/Engine/Source/Programs/Unsync/Private/UnsyncRemote.h b/Engine/Source/Programs/Unsync/Private/UnsyncRemote.h
--- a/Engine/Source/Programs/Unsync/Private/UnsyncRemote.h
+++ b/Engine/Source/Programs/Unsync/Private/UnsyncRemote.h
@@ -43,6 +43,8 @@ struct FRemoteDesc
 	std::string StorageBucket = "unsync";  // TODO: override via command line
 	std::string HttpHeaders;
 
+	std::shared_ptr<FBuffer> Authentication;  // Opaque blob that will be sent to the remote server to authenticate the user
+
 	bool					 bTlsEnable			   = true;	// Prefer TLS, if supported by protocol and remote server
 	bool					 bTlsVerifyCertificate = true;	// Disabling this allows self-signed certificates
 	std::string				 TlsSubject;					// Use host by default
diff --git a/Engine/Source/Programs/Unsync/Private/UnsyncSocket.cpp b/Engine/Source/Programs/Unsync/Private/UnsyncSocket.cpp
--- a/Engine/Source/Programs/Unsync/Private/UnsyncSocket.cpp
+++ b/Engine/Source/Programs/Unsync/Private/UnsyncSocket.cpp
@@ -24,6 +24,7 @@ UNSYNC_THIRD_PARTY_INCLUDES_START
 #endif	// UNSYNC_PLATFORM_UNIX
 
 #include <unordered_set>
+#include <limits>
 
 #if UNSYNC_USE_TLS
 #	include <tls.h>
@@ -89,6 +90,65 @@ GetCurrentHostName()
 	}
 }
 
+static int32
+GetLastSocketError()
+{
+#if UNSYNC_PLATFORM_WINDOWS
+	return WSAGetLastError();
+#else
+	return -1; // TODO: report unix socket error
+#endif
+}
+
+FSocketHandle
+SocketListenTcp(const char* Address, uint16 Port)
+{
+	LazyInitSockets();
+
+	SOCKET ListenSocket = socket(AF_INET, SOCK_STREAM, 0);
+
+	if (ListenSocket == INVALID_SOCKET)
+	{
+		UNSYNC_ERROR(L"Failed to create TCP socket (error code %d)", GetLastSocketError());
+		return 0;
+	}
+
+	sockaddr_in Service;
+	Service.sin_family		= AF_INET;
+	Service.sin_addr.s_addr = inet_addr(Address);
+	Service.sin_port		= htons(Port);
+
+	int32 BindResult = bind(ListenSocket, (SOCKADDR*)&Service, sizeof(Service));
+	if (BindResult == SOCKET_ERROR)
+	{
+		UNSYNC_ERROR(L"Failed to bind TCP socket (error code %d)", GetLastSocketError());
+		SocketClose(ListenSocket);
+		return 0;
+	}
+
+	int32 ListenResult = listen(ListenSocket, 1);
+	if (ListenResult == SOCKET_ERROR)
+	{
+		UNSYNC_ERROR(L"Failed to listen on TCP socket (error code %d)", GetLastSocketError());
+		SocketClose(ListenSocket);
+		return 0;
+	}
+
+	return ListenSocket;
+}
+
+FSocketHandle
+SocketAccept(FSocketHandle ListenSocket)
+{
+	SOCKET AcceptSocket = accept(ListenSocket, nullptr, nullptr);
+	if (AcceptSocket == INVALID_SOCKET)
+	{
+		UNSYNC_ERROR(L"Failed to accept connection on TCP socket (error code %d)", GetLastSocketError());
+		return 0;
+	}
+	return AcceptSocket;
+}
+
 FSocketHandle
 SocketConnectTcp(const char* DestAddress, uint16 Port)
 {
@@ -98,11 +158,7 @@ SocketConnectTcp(const char* DestAddress, uint16 Port)
 
 	if (Sock == INVALID_SOCKET)
 	{
-#if UNSYNC_PLATFORM_WINDOWS
-		UNSYNC_ERROR(L"Failed to create TCP socket (error code %d)", WSAGetLastError());
-#else
-		UNSYNC_ERROR(L"Failed to create TCP socket");  // TODO: report unix socket error
-#endif	// UNSYNC_PLATFORM_WINDOWS
+		UNSYNC_ERROR(L"Failed to create TCP socket (error code %d)", GetLastSocketError());
 		return 0;
 	}
 
@@ -411,13 +467,20 @@ ToString(ESocketSecurity Security)
 }
 
 bool
-SendBuffer(FSocketBase& Socket, const FBuffer& Data)
+SendBuffer(FSocketBase& Socket, const FBufferView& Data)
 {
-	int32 Size = int32(Data.Size());
+	UNSYNC_ASSERT(Data.Size <= std::numeric_limits<int32>::max());
+	int32 Size = int32(Data.Size);
 	bool  bOk  = true;
 	bOk &= SocketSendT(Socket, Size);
-	bOk &= (SocketSend(Socket, Data.Data(), Size) == Size);
+	bOk &= (SocketSend(Socket, Data.Data, Size) == Size);
 	return bOk;
 }
 
+bool
+SendBuffer(FSocketBase& Socket, const FBuffer& Data)
+{
+	return SendBuffer(Socket, Data.View());
+}
+
 }  // namespace unsync
diff --git a/Engine/Source/Programs/Unsync/Private/UnsyncSocket.h b/Engine/Source/Programs/Unsync/Private/UnsyncSocket.h
--- a/Engine/Source/Programs/Unsync/Private/UnsyncSocket.h
+++ b/Engine/Source/Programs/Unsync/Private/UnsyncSocket.h
@@ -10,11 +10,8 @@ struct tls;
 
 namespace unsync {
 
-#ifndef UNSYNC_USE_TLS
-#	define UNSYNC_USE_TLS 0
-#endif
-
 class FBuffer;
+struct FBufferView;
 using FSocketHandle = uintptr_t;
 
 struct FSocketAddress  // opaque buffer that holds socket address
@@ -37,6 +34,9 @@ struct FSocketAddress  // opaque buffer that holds socket address
 // Returns current host name or empty string if it could not be determined.
 std::string GetCurrentHostName();
 
+FSocketHandle SocketListenTcp(const char* Address, uint16 Port);
+FSocketHandle SocketAccept(FSocketHandle ListenSocket);
+
 FSocketHandle SocketConnectTcp(const char* DestAddress, uint16 Port);
 void		  SocketClose(FSocketHandle Socket);
 
@@ -173,6 +173,7 @@ SendStruct(FSocketBase& Socket, const T& Data)
 	return bOk;
 }
 
+bool SendBuffer(FSocketBase& Socket, const FBufferView& Data);
 bool SendBuffer(FSocketBase& Socket, const FBuffer& Data);
 
 template<typename T>
diff --git a/Engine/Source/Programs/Unsync/Private/UnsyncUtil.cpp b/Engine/Source/Programs/Unsync/Private/UnsyncUtil.cpp
--- a/Engine/Source/Programs/Unsync/Private/UnsyncUtil.cpp
+++ b/Engine/Source/Programs/Unsync/Private/UnsyncUtil.cpp
@@ -5,6 +5,7 @@
 
 #if UNSYNC_PLATFORM_WINDOWS
 #	include <Windows.h>
+#	include <shellapi.h>
 #	include <lm.h>
 #	include <lmdfs.h>
 #	include <wincrypt.h>
@@ -15,6 +16,7 @@
 #	pragma comment(lib, "Mpr.lib")
 #endif	// UNSYNC_PLATFORM_WINDOWS
 
+#include <stdlib.h>
 #include <codecvt>
 #include <filesystem>
 #include <locale>
@@ -410,4 +412,45 @@ GetSystemRootCerts()
 	return GSystemRootCerts;
 }
 
+#if UNSYNC_PLATFORM_WINDOWS
+void OpenUrlInDefaultBrowser(const char* Address)
+{
+	ShellExecuteA(nullptr, "open", Address, nullptr, nullptr, SW_SHOWNORMAL);
+}
+#else // UNSYNC_PLATFORM_WINDOWS
+void
+OpenUrlInDefaultBrowser(const char* Address)
+{
+	UNSYNC_FATAL(L"OpenUrlInDefaultBrowser is not implemented");
+}
+#endif // UNSYNC_PLATFORM_WINDOWS
+
+
+#if UNSYNC_PLATFORM_WINDOWS
+FPath GetUserHomeDirectory()
+{
+	if (const char* EnvUserProfile = getenv("USERPROFILE"))
+	{
+		return NormalizeFilenameUtf8(EnvUserProfile);
+	}
+	else
+	{
+		return {};
+	}
+}
+#else // UNSYNC_PLATFORM_WINDOWS
+FPath
+GetUserHomeDirectory()
+{
+	if (const char* EnvUserProfile = getenv("HOME"))
+	{
+		return NormalizeFilenameUtf8(EnvUserProfile);
+	}
+	else
+	{
+		return {};
+	}
+}
+#endif // UNSYNC_PLATFORM_WINDOWS
+
 }  // namespace unsync
diff --git a/Engine/Source/Programs/Unsync/Private/UnsyncUtil.h b/Engine/Source/Programs/Unsync/Private/UnsyncUtil.h
--- a/Engine/Source/Programs/Unsync/Private/UnsyncUtil.h
+++ b/Engine/Source/Programs/Unsync/Private/UnsyncUtil.h
@@ -317,4 +317,8 @@ Append(fmt::memory_buffer& Buf, const char* S)
 	Buf.append(std::string_view(S));
 }
 
+void OpenUrlInDefaultBrowser(const char* Address);
+
+FPath GetUserHomeDirectory();
+
 }  // namespace unsync
