unsync - Automatically refresh access token only when it's close to expiration

- Add FHttpRequest::BearerToken field
- Validate credentials using login endpoint

diff --git a/Engine/Source/Programs/Unsync/Private/UnsyncAuth.cpp b/Engine/Source/Programs/Unsync/Private/UnsyncAuth.cpp
--- a/Engine/Source/Programs/Unsync/Private/UnsyncAuth.cpp
+++ b/Engine/Source/Programs/Unsync/Private/UnsyncAuth.cpp
@@ -7,6 +7,7 @@
 #include "UnsyncProxy.h"
 
 #include <fmt/format.h>
+#include <ctime>
 #include <json11.hpp>
 
 #if UNSYNC_USE_TLS
@@ -112,6 +113,11 @@ TransformBase64UrlSafeToVanilla(std::string& Data)
 {
 	std::replace(Data.begin(), Data.end(), '-', '+');
 	std::replace(Data.begin(), Data.end(), '_', '/');
+
+	while ((Data.length() % 4) != 0)
+	{
+		Data.push_back('=');
+	}
 }
 
 std::string
@@ -222,6 +228,54 @@ StartHttpCallbackServer(FSocketHandle	   CallbackListenSocket,
 	});
 };
 
+TResult<json11::Json>
+DecodeJwtPayload(std::string JwtDataBase64Url)
+{
+	using namespace json11;
+
+	size_t PayloadOffset = JwtDataBase64Url.find('.');
+	if (PayloadOffset == std::string::npos)
+	{
+		return AppError(L"Failed to locate JWT payload section");
+	}
+	PayloadOffset += 1;	 // skip the delimiter
+
+	size_t SignatureOffset = JwtDataBase64Url.find('.', PayloadOffset + 1);
+	if (SignatureOffset == std::string::npos)
+	{
+		return AppError(L"Failed to locate JWT signature section");
+	}
+	SignatureOffset += 1;  // skip the delimiter
+
+	size_t PayloadLength = SignatureOffset - PayloadOffset - 1;
+
+	std::string JwtPayloadBase64 = JwtDataBase64Url.substr(PayloadOffset, PayloadLength);
+	TransformBase64UrlSafeToVanilla(JwtPayloadBase64);
+
+	FBuffer JasonData;
+	bool	bDecoded = DecodeBase64(JwtPayloadBase64, JasonData);
+
+	if (!bDecoded)
+	{
+		return AppError(L"Failed to decode Base64 JWT data");
+	}
+
+	if (*JasonData.end() != 0)
+	{
+		JasonData.PushBack(0);
+	}
+
+	std::string JsonErrorString;
+	Json		JsonObject = Json::parse((const char*)JasonData.Data(), JsonErrorString);
+
+	if (!JsonErrorString.empty())
+	{
+		return AppError(fmt::format("JSON error while parsing token: {}", JsonErrorString.c_str()));
+	}
+
+	return ResultOk(std::move(JsonObject));
+}
+
 TResult<FAuthToken>
 AcquireAuthToken(const FAuthDesc& AuthDesc)
 {
@@ -273,7 +327,7 @@ AcquireAuthToken(const FAuthDesc& AuthDesc)
 		"client_id={}&"
 		"{}"  // optional audience parameter
 		"response_type=code&"
-		"scope=offline_access%20profile%20openid%20email&"
+		"scope=offline_access&"
 		"code_challenge_method=S256&"
 		"code_challenge={}&"
 		"state={}&"
@@ -355,6 +409,16 @@ AcquireAuthToken(const FAuthDesc& AuthDesc)
 			TokenType		 = JsonObject["token_type"].string_value();
 			ExpiresInSeconds = int64(JsonObject["expires_in"].number_value());
 
+			TResult<json11::Json> DecodedAccessTokenResult = DecodeJwtPayload(AccessToken);
+			if (DecodedAccessTokenResult.IsOk())
+			{
+				const json11::Json& AccessTokenJsonObject = DecodedAccessTokenResult.GetData();
+				if (auto& Field = AccessTokenJsonObject["exp"]; Field.is_number())
+				{
+					Result.ExirationTime = int64(Field.number_value());
+				}
+			}
+
 			Result.Raw = JsonString;
 		}
 		else
@@ -377,8 +441,6 @@ AcquireAuthToken(const FAuthDesc& AuthDesc)
 TResult<FAuthUserInfo>
 GetUserInfo(FHttpConnection& HttpConnection, const FAuthDesc& AuthDesc, const FAuthToken& AuthToken)
 {
-	std::string AuthHeader = fmt::format("Authorization: Bearer {}", AuthToken.Access);
-
 	if (AuthDesc.UserInfoEndpoint.empty())
 	{
 		return AppError(L"User info endpoint is unknown");
@@ -387,7 +449,7 @@ GetUserInfo(FHttpConnection& HttpConnection, const FAuthDesc& AuthDesc, const FA
 	FHttpRequest Request;
 	Request.Url			  = AuthDesc.UserInfoEndpoint;
 	Request.Method		  = EHttpMethod::GET;
-	Request.CustomHeaders = AuthHeader;
+	Request.BearerToken	  = AuthToken.Access;
 
 	FHttpResponse Response = HttpRequest(HttpConnection, Request);
 
@@ -476,6 +538,16 @@ RefreshAuthToken(const FAuthDesc& AuthDesc, const FAuthToken& PreviousToken)
 			TokenType		 = JsonObject["token_type"].string_value();
 			ExpiresInSeconds = int64(JsonObject["expires_in"].number_value());
 
+			TResult<json11::Json> DecodedAccessTokenResult = DecodeJwtPayload(AccessToken);
+			if (DecodedAccessTokenResult.IsOk())
+			{
+				const json11::Json& AccessTokenJsonObject = DecodedAccessTokenResult.GetData();
+				if (auto& Field = AccessTokenJsonObject["exp"]; Field.is_number())
+				{
+					Result.ExirationTime = int64(Field.number_value());
+				}
+			}
+
 			Result.Raw = JsonString;
 		}
 		else
@@ -507,19 +579,43 @@ GenerateTokenId(const FRemoteDesc& RemoteDesc)
 }
 
 bool
-SaveRefreshToken(const FPath& Path, const FAuthToken& AuthToken)
+SaveAuthToken(const FPath& Path, const FAuthToken& AuthToken)
 {
-	return WriteBufferToFile(Path, (const uint8*)AuthToken.Refresh.data(), AuthToken.Refresh.length());
+	return WriteBufferToFile(Path, (const uint8*)AuthToken.Raw.data(), AuthToken.Raw.length());
 }
 
 TResult<FAuthToken>
-LoadRefreshToken(const FPath& Path)
+LoadAuthToken(const FPath& Path)
 {
 	FBuffer FileBuffer = ReadFileToBuffer(Path);
 	if (FileBuffer.Size())
 	{
+		using namespace json11;
+
 		FAuthToken AuthToken;
-		AuthToken.Refresh.append((const char*)FileBuffer.Data(), FileBuffer.Size());
+		AuthToken.Raw.append((const char*)FileBuffer.Data(), FileBuffer.Size());
+
+		std::string JsonErrorString;
+		Json		JsonObject = Json::parse(AuthToken.Raw, JsonErrorString);
+
+		if (!JsonErrorString.empty())
+		{
+			return AppError(fmt::format("JSON error while parsing token: {}", JsonErrorString.c_str()));
+		}
+
+		AuthToken.Access  = JsonObject["access_token"].string_value();
+		AuthToken.Refresh = JsonObject["refresh_token"].string_value();
+
+		TResult<json11::Json> DecodedAccessTokenResult = DecodeJwtPayload(AuthToken.Access);
+		if (DecodedAccessTokenResult.IsOk())
+		{
+			const json11::Json& AccessTokenJsonObject = DecodedAccessTokenResult.GetData();
+			if (auto& Field = AccessTokenJsonObject["exp"]; Field.is_number())
+			{
+				AuthToken.ExirationTime = int64(Field.number_value());
+			}
+		}
+
 		return ResultOk(AuthToken);
 	}
 	else
@@ -545,37 +641,87 @@ RefreshOrAcquireToken(const FAuthDesc& AuthDesc, const FAuthToken& PreviousToken
 	return AcquireAuthToken(AuthDesc);
 }
 
-TResult<FAuthToken>
-Authenticate(const FRemoteDesc& RemoteDesc, const FAuthDesc& AuthDesc)
+TResult<FPath>
+GetTokenCachePath(const FRemoteDesc& RemoteDesc)
 {
 	FPath UserHomePath = GetUserHomeDirectory();
+	if (UserHomePath.empty())
+	{
+		return AppError(L"Could not query user home directory path");
+	}
 
 	std::string TokenId = GenerateTokenId(RemoteDesc);
 
 	FPath UnsyncSettingsPath = UserHomePath / FPath(".unsync");
 	FPath TokenCachePath	 = UnsyncSettingsPath / FPath(TokenId);
 
+	return ResultOk(TokenCachePath);
+}
+
+void
+LogAuthTokenExpiration(const FAuthToken& AuthToken)
+{
+	if (AuthToken.ExirationTime != 0)
+	{
+		int64 CurrentTime	   = GetSecondsFromUnixEpoch();
+		int64 ExpiresInSeconds = AuthToken.ExirationTime - CurrentTime;
+		if (ExpiresInSeconds > 0)
+		{
+			UNSYNC_VERBOSE(L"Authentication token will expire in %d sec", int(ExpiresInSeconds));
+		}
+		else
+		{
+			UNSYNC_VERBOSE(L"Authentication token has expired");
+		}
+	}
+}
+
+TResult<FAuthToken>
+Authenticate(const FRemoteDesc& RemoteDesc, int32 RefreshThreshold)
+{
+	// Authentication must be serialized (only one thread should ever open the browser for interactive login, etc.)
+	static std::mutex			AuthMutex;
+	std::lock_guard<std::mutex> LockGuard(AuthMutex);
+
 	FAuthToken PreviousToken;
 
-	if (PathExists(TokenCachePath))
+	TResult<FPath> TokenCachePathResult = GetTokenCachePath(RemoteDesc);
+
+	if (const FPath* TokenCachePath = TokenCachePathResult.TryData())
 	{
-		TResult<FAuthToken> LoadResult = LoadRefreshToken(TokenCachePath);
+		TResult<FAuthToken> LoadResult = LoadAuthToken(*TokenCachePath);
 		if (FAuthToken* LoadedToken = LoadResult.TryData())
 		{
-			UNSYNC_VERBOSE(L"Loaded cached authentication data");
+			UNSYNC_VERBOSE2(L"Loaded cached authentication token");
 			PreviousToken = std::move(*LoadedToken);
 		}
 	}
 
+	int64 CurrentTime	   = GetSecondsFromUnixEpoch();
+	int64 ExpiresInSeconds = PreviousToken.ExirationTime - CurrentTime;
+	if (ExpiresInSeconds > RefreshThreshold)
+	{
+		LogAuthTokenExpiration(PreviousToken);
+		return ResultOk(PreviousToken);
+	}
+
+	TResult<FAuthDesc> AuthDescResult = GetAuthenticationDesc(RemoteDesc);
+	if (AuthDescResult.IsError())
+	{
+		return MoveError<FAuthToken>(AuthDescResult);
+	}
+
+	const FAuthDesc& AuthDesc = AuthDescResult.GetData();
+
 	TResult<FAuthToken> FreshTokenResult = RefreshOrAcquireToken(AuthDesc, PreviousToken);
 	if (FreshTokenResult.IsError())
 	{
 		return FreshTokenResult;
 	}
 
-	if (!UserHomePath.empty())
+	if (const FPath* TokenCachePath = TokenCachePathResult.TryData())
 	{
-		CreateDirectories(UnsyncSettingsPath);
+		CreateDirectories(TokenCachePath->parent_path());
 
 		// Allow saving tokens during dry run
 		// TODO: need a dedicated file flag to allow writes during dry run
@@ -583,42 +729,37 @@ Authenticate(const FRemoteDesc& RemoteDesc, const FAuthDesc& AuthDesc)
 		const bool bPrevDryRun = GDryRun;
 		GDryRun				   = false;
 
-		bool bSaved = SaveRefreshToken(TokenCachePath, FreshTokenResult.GetData());
+		bool bSaved = SaveAuthToken(*TokenCachePath, FreshTokenResult.GetData());
 
 		GDryRun = bPrevDryRun;
 
 		if (bSaved)
 		{
-			UNSYNC_VERBOSE2(L"Saved refresh token to file: %s", TokenCachePath.wstring().c_str());
+			UNSYNC_VERBOSE2(L"Saved authentication token to file: %s", TokenCachePath->wstring().c_str());
 		}
 	}
 
+	if (FreshTokenResult.IsOk())
+	{
+		LogAuthTokenExpiration(FreshTokenResult.GetData());
+	}
+
 	return FreshTokenResult;
 }
 
 bool
 TryAddAuthentication(FRemoteDesc& InOutRemoteDesc)
 {
-	TResult<FAuthDesc> AuthDescResult = GetAuthenticationDesc(InOutRemoteDesc);
-	if (AuthDescResult.IsError())
-	{
-		return false;
-	}
+	TResult<FAuthToken> AuthTokenResult = Authenticate(InOutRemoteDesc, 5 * 60);
 
-	const FAuthDesc& AuthDesc = AuthDescResult.GetData();
-
-	TResult<FAuthToken> AuthTokenResult = Authenticate(InOutRemoteDesc, AuthDesc);
 	if (AuthTokenResult.IsError())
 	{
 		return false;
 	}
 
-	const std::string& AccessToken = AuthTokenResult->Access;
-
 	// Authentication requires encrypted connection
-	InOutRemoteDesc.bTlsEnable	   = true;
-	InOutRemoteDesc.Authentication = std::make_shared<FBuffer>();
-	InOutRemoteDesc.Authentication->Append((const uint8*)AccessToken.data(), AccessToken.length());
+	InOutRemoteDesc.bTlsEnable				= true;
+	InOutRemoteDesc.bAuthenticationRequired = true;
 
 	return true;
 }
@@ -711,6 +852,12 @@ GetAuthenticationDesc(const FRemoteDesc& RemoteDesc)
 	return ResultOk(AuthDesc);
 }
 
+int64
+GetSecondsFromUnixEpoch()
+{
+	return int64(std::time(nullptr));
+}
+
 }  // namespace unsync
 
 #endif	// UNSYNC_USE_TLS
diff --git a/Engine/Source/Programs/Unsync/Private/UnsyncAuth.h b/Engine/Source/Programs/Unsync/Private/UnsyncAuth.h
--- a/Engine/Source/Programs/Unsync/Private/UnsyncAuth.h
+++ b/Engine/Source/Programs/Unsync/Private/UnsyncAuth.h
@@ -57,7 +57,8 @@ TResult<FAuthDesc> GetAuthenticationDesc(const FRemoteDesc& RemoteDesc);
 // - Attempt use stored refresh token first
 // - If refresh is not possible, use PKCE Authentication flow to get new tokens
 // - Save refresh token in user directory for future use
-TResult<FAuthToken> Authenticate(const FRemoteDesc& RemoteDesc, const FAuthDesc& AuthDesc);
+// - Skips acquiring new token if remaining valid time is above RefreshThreshold (in seconds)
+TResult<FAuthToken> Authenticate(const FRemoteDesc& RemoteDesc, int32 RefreshThreshold = INT_MAX);
 
 // Auth utility functions
 
@@ -78,8 +79,8 @@ TResult<FAuthUserInfo> GetUserInfo(FHttpConnection& HttpConnection, const FAuthD
 
 std::string GenerateTokenId(const FRemoteDesc& RemoteDesc);
 
-bool				SaveRefreshToken(const FPath& Path, const FAuthToken& AuthToken);
-TResult<FAuthToken> LoadRefreshToken(const FPath& Path);
+bool				SaveAuthToken(const FPath& Path, const FAuthToken& AuthToken);
+TResult<FAuthToken> LoadAuthToken(const FPath& Path);
 
 std::string SecureRandomBytesAsHexString(uint32 NumBytes);
 FHash256	HashSha256Bytes(const uint8* Data, uint64 Size);
@@ -93,6 +94,8 @@ std::string GetPKCECodeChallenge(std::string_view CodeVerifier);
 
 bool TryAddAuthentication(FRemoteDesc& InOutRemoteDesc);
 
+int64 GetSecondsFromUnixEpoch();
+
 }  // namespace unsync
 
 #endif	// UNSYNC_USE_TLS
diff --git a/Engine/Source/Programs/Unsync/Private/UnsyncCmdQuery.cpp b/Engine/Source/Programs/Unsync/Private/UnsyncCmdQuery.cpp
--- a/Engine/Source/Programs/Unsync/Private/UnsyncCmdQuery.cpp
+++ b/Engine/Source/Programs/Unsync/Private/UnsyncCmdQuery.cpp
@@ -146,20 +146,43 @@ CmdQueryMirrors(const FCmdQueryOptions& Options)
 int32
 CmdQueryLogin(const FCmdQueryOptions& Options)
 {
-	TResult<FAuthDesc> AuthDescResult = GetAuthenticationDesc(Options.Remote);
-	if (AuthDescResult.IsError())
+	TResult<FAuthToken> AuthToken = Authenticate(Options.Remote, 5 * 60);
+
+	if (AuthToken.IsOk())
 	{
-		LogError(AuthDescResult.GetError());
-		return -1;
-	}
+		{
+			FHttpConnection Connection		 = FHttpConnection::CreateDefaultHttps(Options.Remote.HostAddress, Options.Remote.HostPort);
+			Connection.bTlsVerifyCertificate = false;
 
-	const FAuthDesc& AuthDesc = AuthDescResult.GetData();
+			FHttpRequest Request;
+			Request.Url			   = "/api/v1/login";
+			Request.Method		   = EHttpMethod::GET;
+			Request.BearerToken	   = AuthToken->Access;
+			FHttpResponse Response = HttpRequest(Connection, Request);
 
-	TResult<FAuthToken> AuthToken = Authenticate(Options.Remote, AuthDesc);
+			if (Response.Success())
+			{
+				UNSYNC_LOG("Login successful");
+			}
+			else
+			{
+				LogError(HttpError(Response.Code));
+				return -1;
+			}
+		}
 
-	if (AuthToken.IsOk())
-	{
-		UNSYNC_LOG(L"Login successful");
+
+		UNSYNC_LOG(L"Getting user info from authentication server");
+		UNSYNC_LOG_INDENT;
+
+		TResult<FAuthDesc> AuthDescResult = GetAuthenticationDesc(Options.Remote);
+		if (AuthDescResult.IsError())
+		{
+			LogError(AuthDescResult.GetError());
+			return -1;
+		}
+
+		const FAuthDesc& AuthDesc = AuthDescResult.GetData();
 
 		FHttpConnection AuthConnection = FHttpConnection::CreateDefaultHttps(AuthDesc.ServerHost);
 
diff --git a/Engine/Source/Programs/Unsync/Private/UnsyncHttp.cpp b/Engine/Source/Programs/Unsync/Private/UnsyncHttp.cpp
--- a/Engine/Source/Programs/Unsync/Private/UnsyncHttp.cpp
+++ b/Engine/Source/Programs/Unsync/Private/UnsyncHttp.cpp
@@ -306,6 +306,13 @@ HttpRequestBegin(FHttpConnection& Connection, const FHttpRequest& Request)
 		}
 	}
 
+	if (!Request.BearerToken.empty())
+	{
+		HttpHeader += "Authorization: Bearer ";
+		HttpHeader += Request.BearerToken;
+		HttpHeader += "\r\n";
+	}
+
 	HttpHeader += "Host: " + Connection.HostAddress + "\r\n";
 
 	HttpHeader += "User-Agent: unsync v" + GetVersionString() + "\r\n";
diff --git a/Engine/Source/Programs/Unsync/Private/UnsyncHttp.h b/Engine/Source/Programs/Unsync/Private/UnsyncHttp.h
--- a/Engine/Source/Programs/Unsync/Private/UnsyncHttp.h
+++ b/Engine/Source/Programs/Unsync/Private/UnsyncHttp.h
@@ -39,6 +39,7 @@ struct FHttpRequest
 	EHttpMethod		 Method		   = EHttpMethod::GET;
 	std::string_view Url		   = {};
 	std::string_view CustomHeaders = {};
+	std::string_view BearerToken   = {};
 
 	EHttpContentType PayloadContentType = EHttpContentType::Unknown;
 	FBufferView		 Payload			= {};
diff --git a/Engine/Source/Programs/Unsync/Private/UnsyncProxy.cpp b/Engine/Source/Programs/Unsync/Private/UnsyncProxy.cpp
--- a/Engine/Source/Programs/Unsync/Private/UnsyncProxy.cpp
+++ b/Engine/Source/Programs/Unsync/Private/UnsyncProxy.cpp
@@ -4,6 +4,7 @@
 #include "UnsyncCompression.h"
 #include "UnsyncFile.h"
 #include "UnsyncJupiter.h"
+#include "UnsyncAuth.h"
 
 #include <json11.hpp>
 #include <fmt/format.h>
@@ -122,26 +123,39 @@ FUnsyncProtocolImpl::FUnsyncProtocolImpl(const FRemoteDesc&				RemoteDesc,
 		}();
 	}
 
-	if (IsValid() && Features.bAuthentication && RemoteDesc.Authentication)
+	if (IsValid() && Features.bAuthentication && RemoteDesc.bAuthenticationRequired)
 	{
 		bool bOk = IsValid();
 
-		FCommandPacket Packet;
-		Packet.CommandId = COMMAND_ID_AUTHENTICATE;
-		bOk &= SendStruct(*SocketHandle, Packet);
-		bOk &= SendBuffer(*SocketHandle, *RemoteDesc.Authentication);
+		TResult<FAuthToken> AuthTokenResult = Authenticate(RemoteDesc, 15 * 60);
+
+		if (AuthTokenResult.IsOk())
+		{
+			FBufferView AccessToken = {(const uint8*)AuthTokenResult->Access.data(), AuthTokenResult->Access.length()};
+
+			FCommandPacket Packet;
+			Packet.CommandId = COMMAND_ID_AUTHENTICATE;
+			bOk &= SendStruct(*SocketHandle, Packet);
+			bOk &= SendBuffer(*SocketHandle, AccessToken);
 
-		int32 ResultSize = 0;
-		bOk &= SocketRecvT(*SocketHandle, ResultSize);
+			int32 ResultSize = 0;
+			bOk &= SocketRecvT(*SocketHandle, ResultSize);
 
-		FBuffer ResultBuffer;
-		if (ResultSize)
+			FBuffer ResultBuffer;
+			if (ResultSize)
+			{
+				ResultBuffer.Resize(ResultSize);
+				bOk &= (SocketRecvAll(*SocketHandle, ResultBuffer.Data(), ResultSize) == ResultSize);
+			}
+
+			// TODO: parse authentication result packet and report errors
+		}
+		else
 		{
-			ResultBuffer.Resize(ResultSize);
-			bOk &= (SocketRecvAll(*SocketHandle, ResultBuffer.Data(), ResultSize) == ResultSize);
+			LogError(AuthTokenResult.GetError());
+			UNSYNC_ERROR(L"Server requires authentication");
+			Invalidate();
 		}
-
-		// TODO: parse authentication result packet and report errors
 	}
 }
 
diff --git a/Engine/Source/Programs/Unsync/Private/UnsyncRemote.h b/Engine/Source/Programs/Unsync/Private/UnsyncRemote.h
--- a/Engine/Source/Programs/Unsync/Private/UnsyncRemote.h
+++ b/Engine/Source/Programs/Unsync/Private/UnsyncRemote.h
@@ -15,15 +15,13 @@ class FBuffer;
 
 static constexpr uint16 UNSYNC_DEFAULT_PORT = 53841;
 
-enum class EProtocolFlavor
-{
+enum class EProtocolFlavor {
 	Unknown,
 	Unsync,
 	Jupiter,
 };
 
-enum class ETransportProtocol
-{
+enum class ETransportProtocol {
 	Http,
 	Unsync,
 };
@@ -43,13 +41,13 @@ struct FRemoteDesc
 	std::string StorageBucket = "unsync";  // TODO: override via command line
 	std::string HttpHeaders;
 
-	std::shared_ptr<FBuffer> Authentication;  // Opaque blob that will be sent to the remote server to authenticate the user
-
 	bool					 bTlsEnable			   = true;	// Prefer TLS, if supported by protocol and remote server
 	bool					 bTlsVerifyCertificate = true;	// Disabling this allows self-signed certificates
 	std::string				 TlsSubject;					// Use host by default
 	std::shared_ptr<FBuffer> TlsCacert;	 // Custom CA to use for server certificate validation (system root CA is used by default)
 
+	bool bAuthenticationRequired = false;
+
 	uint32 MaxConnections = 8;	// Limit on concurrent connections to this server
 
 	bool IsValid() const { return Protocol != EProtocolFlavor::Unknown && !HostAddress.empty() && HostPort != 0; }
