unsync - Groundwork for list query

diff --git a/Engine/Source/Programs/Unsync/Private/UnsyncAuth.cpp b/Engine/Source/Programs/Unsync/Private/UnsyncAuth.cpp
--- a/Engine/Source/Programs/Unsync/Private/UnsyncAuth.cpp
+++ b/Engine/Source/Programs/Unsync/Private/UnsyncAuth.cpp
@@ -260,10 +260,7 @@ DecodeJwtPayload(std::string JwtDataBase64Url)
 		return AppError(L"Failed to decode Base64 JWT data");
 	}
 
-	if (*JasonData.end() != 0)
-	{
-		JasonData.PushBack(0);
-	}
+	JasonData.PushBack(0);
 
 	std::string JsonErrorString;
 	Json		JsonObject = Json::parse((const char*)JasonData.Data(), JsonErrorString);
diff --git a/Engine/Source/Programs/Unsync/Private/UnsyncCmdQuery.cpp b/Engine/Source/Programs/Unsync/Private/UnsyncCmdQuery.cpp
--- a/Engine/Source/Programs/Unsync/Private/UnsyncCmdQuery.cpp
+++ b/Engine/Source/Programs/Unsync/Private/UnsyncCmdQuery.cpp
@@ -151,8 +151,7 @@ CmdQueryLogin(const FCmdQueryOptions& Options)
 	if (AuthToken.IsOk())
 	{
 		{
-			FHttpConnection Connection		 = FHttpConnection::CreateDefaultHttps(Options.Remote.HostAddress, Options.Remote.HostPort);
-			Connection.bTlsVerifyCertificate = false;
+			FHttpConnection Connection = FHttpConnection::CreateDefaultHttps(Options.Remote);
 
 			FHttpRequest Request;
 			Request.Url			   = "/api/v1/login";
@@ -171,7 +170,6 @@ CmdQueryLogin(const FCmdQueryOptions& Options)
 			}
 		}
 
-
 		UNSYNC_LOG(L"Getting user info from authentication server");
 		UNSYNC_LOG_INDENT;
 
@@ -208,6 +206,56 @@ CmdQueryLogin(const FCmdQueryOptions& Options)
 		return -1;
 	}
 }
+
+int32
+CmdQueryList(const FCmdQueryOptions& Options)
+{
+	TResult<FAuthToken> AuthToken = Authenticate(Options.Remote, 5 * 60);
+	if (!AuthToken.IsOk())
+	{
+		LogError(AuthToken.GetError());
+		return -1;
+	}
+
+	FHttpConnection Connection = FHttpConnection::CreateDefaultHttps(Options.Remote);
+
+	std::string Url = fmt::format("/api/v1/list?{}", Options.Args);
+
+	FHttpRequest Request;
+	Request.Url			   = Url;
+	Request.Method		   = EHttpMethod::GET;
+	Request.BearerToken	   = AuthToken->Access;
+	FHttpResponse Response = HttpRequest(Connection, Request);
+
+	Response.Buffer.PushBack(0);
+
+	std::string	 JsonErrorString;
+	json11::Json JsonObject = json11::Json::parse((const char*)Response.Buffer.Data(), JsonErrorString);
+
+	if (!JsonErrorString.empty())
+	{
+		LogError(AppError(fmt::format("JSON error: {}", JsonErrorString.c_str())));
+		return -1;
+	}
+	/*
+	if (JsonObject.is_array())
+	{
+		int32 NumFiles = int32(JsonObject.array_items().size());
+		UNSYNC_LOG("Found files: %d", NumFiles);
+		for (auto& FileIt : JsonObject.array_items())
+		{
+			if (auto Value = FileIt["path"]; Value.is_string())
+			{
+				UNSYNC_LOG(L" %hs", Value.string_value().c_str());
+			}
+		}
+	}
+	*/
+
+	wprintf(L"%hs\n", Response.Buffer.Data());
+
+	return 0;
+}
 #endif	// UNSYNC_USE_TLS
 
 int32
@@ -222,6 +270,10 @@ CmdQuery(const FCmdQueryOptions& Options)
 	{
 		return CmdQueryLogin(Options);
 	}
+	else if (Options.Query == "list")
+	{
+		return CmdQueryList(Options);
+	}
 #endif	// UNSYNC_USE_TLS
 	else
 	{
diff --git a/Engine/Source/Programs/Unsync/Private/UnsyncCmdQuery.h b/Engine/Source/Programs/Unsync/Private/UnsyncCmdQuery.h
--- a/Engine/Source/Programs/Unsync/Private/UnsyncCmdQuery.h
+++ b/Engine/Source/Programs/Unsync/Private/UnsyncCmdQuery.h
@@ -10,6 +10,7 @@ namespace unsync {
 struct FCmdQueryOptions
 {
 	std::string Query;
+	std::string Args;
 	FRemoteDesc Remote;
 };
 
diff --git a/Engine/Source/Programs/Unsync/Private/UnsyncHttp.cpp b/Engine/Source/Programs/Unsync/Private/UnsyncHttp.cpp
--- a/Engine/Source/Programs/Unsync/Private/UnsyncHttp.cpp
+++ b/Engine/Source/Programs/Unsync/Private/UnsyncHttp.cpp
@@ -501,6 +501,15 @@ FHttpConnection::CreateDefaultHttps(const std::string_view InHostAddress, uint16
 	return FHttpConnection(InHostAddress, Port, &TlsSettings);
 }
 
+FHttpConnection
+FHttpConnection::CreateDefaultHttps(const FRemoteDesc& RemoteDesc)
+{
+	FTlsClientSettings TlsSettings;
+	TlsSettings.Subject			   = RemoteDesc.HostAddress.data();
+	TlsSettings.bVerifyCertificate = RemoteDesc.bTlsVerifyCertificate;
+	return FHttpConnection(RemoteDesc.HostAddress, RemoteDesc.HostPort, &TlsSettings);
+}
+
 bool
 FHttpConnection::Open()
 {
diff --git a/Engine/Source/Programs/Unsync/Private/UnsyncHttp.h b/Engine/Source/Programs/Unsync/Private/UnsyncHttp.h
--- a/Engine/Source/Programs/Unsync/Private/UnsyncHttp.h
+++ b/Engine/Source/Programs/Unsync/Private/UnsyncHttp.h
@@ -66,6 +66,7 @@ struct FHttpConnection
 
 	[[nodiscard]] static FHttpConnection CreateDefaultHttp(const std::string_view InHostAddress, uint16 Port = 80);
 	[[nodiscard]] static FHttpConnection CreateDefaultHttps(const std::string_view InHostAddress, uint16 Port = 443);
+	[[nodiscard]] static FHttpConnection CreateDefaultHttps(const FRemoteDesc& RemoteDesc);
 
 	bool Open();
 	void Close();
diff --git a/Engine/Source/Programs/Unsync/Private/UnsyncMain.cpp b/Engine/Source/Programs/Unsync/Private/UnsyncMain.cpp
--- a/Engine/Source/Programs/Unsync/Private/UnsyncMain.cpp
+++ b/Engine/Source/Programs/Unsync/Private/UnsyncMain.cpp
@@ -71,6 +71,7 @@ InnerMain(int Argc, char** Argv)
 	std::string				 ProtocolName = "jupiter";
 	std::string				 HttpHeaderFilenameUtf8;
 	std::string				 QueryStringUtf8;
+	std::string				 QueryArgsUtf8;
 	std::string				 ScavengeRootUtf8;
 	bool					 bForceOperation	 = false;
 	bool					 bAllowInsecureTls	 = false;
@@ -261,6 +262,7 @@ InnerMain(int Argc, char** Argv)
 
 	CLI::App* SubQuery = Cli.add_subcommand("query", "Run a query command on the remote server");
 	SubQuery->add_option("QueryString", QueryStringUtf8, "Query")->required();
+	SubQuery->add_option("QueryArgs", QueryArgsUtf8, "Query arguments");
 	AddProxyOptions(SubQuery);
 
 #if UNSYNC_USE_TLS
@@ -779,6 +781,7 @@ InnerMain(int Argc, char** Argv)
 	{
 		FCmdQueryOptions QueryOptions;
 		QueryOptions.Query	= QueryStringUtf8;
+		QueryOptions.Args	= QueryArgsUtf8;
 		QueryOptions.Remote = RemoteDesc;
 		return CmdQuery(QueryOptions);
 	}
