unsync - Add dedicated login command and make output more machine-friendly

* Make TLS a mandatory required library
* Redirect human-readable log messages of login and query commands to stderr, while only directing machine-readable messages to stdout

diff --git a/Engine/Source/Programs/Unsync/CMakeLists.txt b/Engine/Source/Programs/Unsync/CMakeLists.txt
--- a/Engine/Source/Programs/Unsync/CMakeLists.txt
+++ b/Engine/Source/Programs/Unsync/CMakeLists.txt
@@ -15,7 +15,6 @@ endif()
 
 project(unsync)
 
-option(UNSYNC_USE_TLS "Include TLS support for encrypted connections via libretls" ON)
 option(UNSYNC_USE_DEBUG_HEAP "Include support for debug memory allocator" ON)
 
 # Validate compiler
@@ -39,6 +38,8 @@ set(src
 	Private/UnsyncCmdDiff.h
 	Private/UnsyncCmdHash.cpp
 	Private/UnsyncCmdHash.h
+	Private/UnsyncCmdLogin.cpp
+	Private/UnsyncCmdLogin.h
 	Private/UnsyncCmdMount.cpp
 	Private/UnsyncCmdMount.h
 	Private/UnsyncCmdPatch.cpp
@@ -138,13 +139,9 @@ find_path(BLAKE3_INCLUDE_DIR blake3.h REQUIRED)
 target_include_directories(unsync PRIVATE ${BLAKE3_INCLUDE_DIR})
 target_link_libraries(unsync PRIVATE blake3)
 
-if (UNSYNC_USE_TLS)
-	find_path(LIBRESSL_INCLUDE_DIR tls.h)
-	target_include_directories(unsync PRIVATE ${LIBRESSL_INCLUDE_DIR})
-	target_link_libraries(unsync PRIVATE tls ssl crypto)
-	target_compile_definitions(unsync PRIVATE UNSYNC_USE_TLS=1)
-	message(STATUS "TLS support enabled")
-endif(UNSYNC_USE_TLS)
+find_path(LIBRESSL_INCLUDE_DIR tls.h)
+target_include_directories(unsync PRIVATE ${LIBRESSL_INCLUDE_DIR})
+target_link_libraries(unsync PRIVATE tls ssl crypto)
 
 if (UNSYNC_USE_DEBUG_HEAP)
 	target_compile_definitions(unsync PRIVATE UNSYNC_USE_DEBUG_HEAP=1)
diff --git a/Engine/Source/Programs/Unsync/Private/UnsyncAuth.cpp b/Engine/Source/Programs/Unsync/Private/UnsyncAuth.cpp
--- a/Engine/Source/Programs/Unsync/Private/UnsyncAuth.cpp
+++ b/Engine/Source/Programs/Unsync/Private/UnsyncAuth.cpp
@@ -10,16 +10,11 @@
 #include <ctime>
 #include <json11.hpp>
 
-#if UNSYNC_USE_TLS
+#include <openssl/err.h>
+#include <openssl/evp.h>  // Base64 encoding
+#include <openssl/rand.h>
+#include <openssl/sha.h>
 
-#	include <openssl/err.h>
-#	include <openssl/rand.h>
-#	include <openssl/sha.h>
-#	include <openssl/evp.h>  // Base64 encoding
-
-#endif	// UNSYNC_USE_TLS
-
-#if UNSYNC_USE_TLS
 namespace unsync {
 
 std::string
@@ -444,9 +439,9 @@ GetUserInfo(FHttpConnection& HttpConnection, const FAuthDesc& AuthDesc, const FA
 	}
 
 	FHttpRequest Request;
-	Request.Url			  = AuthDesc.UserInfoEndpoint;
-	Request.Method		  = EHttpMethod::GET;
-	Request.BearerToken	  = AuthToken.Access;
+	Request.Url			= AuthDesc.UserInfoEndpoint;
+	Request.Method		= EHttpMethod::GET;
+	Request.BearerToken = AuthToken.Access;
 
 	FHttpResponse Response = HttpRequest(HttpConnection, Request);
 
@@ -856,5 +851,3 @@ GetSecondsFromUnixEpoch()
 }
 
 }  // namespace unsync
-
-#endif	// UNSYNC_USE_TLS
diff --git a/Engine/Source/Programs/Unsync/Private/UnsyncAuth.h b/Engine/Source/Programs/Unsync/Private/UnsyncAuth.h
--- a/Engine/Source/Programs/Unsync/Private/UnsyncAuth.h
+++ b/Engine/Source/Programs/Unsync/Private/UnsyncAuth.h
@@ -9,8 +9,6 @@
 
 #include <string>
 
-#if UNSYNC_USE_TLS
-
 namespace unsync {
 
 struct FHttpConnection;
@@ -97,5 +95,3 @@ bool TryAddAuthentication(FRemoteDesc& InOutRemoteDesc);
 int64 GetSecondsFromUnixEpoch();
 
 }  // namespace unsync
-
-#endif	// UNSYNC_USE_TLS
diff --git a/Engine/Source/Programs/Unsync/Private/UnsyncCmdLogin.cpp b/Engine/Source/Programs/Unsync/Private/UnsyncCmdLogin.cpp
new file mode 100644
--- /dev/null
+++ b/Engine/Source/Programs/Unsync/Private/UnsyncCmdLogin.cpp
@@ -0,0 +1,92 @@
+// Copyright Epic Games, Inc. All Rights Reserved.
+
+#include "UnsyncCmdLogin.h"
+#include "UnsyncAuth.h"
+#include "UnsyncHttp.h"
+#include "UnsyncLog.h"
+#include "UnsyncProxy.h"
+
+namespace unsync {
+
+std::string_view
+GetJwtPayload(std::string_view JwtData)
+{
+	size_t PayloadOffset = JwtData.find('.');
+	if (PayloadOffset == std::string::npos)
+	{
+		return {};
+	}
+	PayloadOffset += 1;	 // skip the delimiter
+
+	size_t SignatureOffset = JwtData.find('.', PayloadOffset + 1);
+	if (SignatureOffset == std::string::npos)
+	{
+		return {};
+	}
+	SignatureOffset += 1;  // skip the delimiter
+
+	size_t PayloadLength = SignatureOffset - PayloadOffset - 1;
+
+	return JwtData.substr(PayloadOffset, PayloadLength);
+}
+
+int32
+CmdLogin(const FCmdLoginOptions& Options)
+{
+	TResult<FAuthToken> AuthTokenResult = Authenticate(Options.Remote, 5 * 60);
+
+	if (AuthTokenResult.IsOk())
+	{
+		{
+			FHttpConnection Connection = FHttpConnection::CreateDefaultHttps(Options.Remote);
+
+			FHttpRequest Request;
+			Request.Url			   = "/api/v1/login";
+			Request.Method		   = EHttpMethod::GET;
+			Request.BearerToken	   = AuthTokenResult->Access;
+			FHttpResponse Response = HttpRequest(Connection, Request);
+
+			if (Response.Success())
+			{
+				UNSYNC_VERBOSE("Login successful");
+			}
+			else
+			{
+				LogError(HttpError(Response.Code));
+				return -1;
+			}
+		}
+
+		if (Options.bPrint)
+		{
+			if (Options.bDecode)
+			{
+				std::string JwtPayload = std::string(GetJwtPayload(AuthTokenResult->Access));
+				TransformBase64UrlSafeToVanilla(JwtPayload);
+
+				FBuffer DecodedTokenData;
+				if (!DecodeBase64(JwtPayload, DecodedTokenData))
+				{
+					UNSYNC_ERROR("Failed to Base64-decode access token");
+					return -1;
+				}
+				DecodedTokenData.Append(0);
+
+				LogPrintf(ELogLevel::MachineReadable, L"%hs\n", (const char*)DecodedTokenData.Data());
+			}
+			else
+			{
+				LogPrintf(ELogLevel::MachineReadable, L"%hs\n", AuthTokenResult->Raw.c_str());
+			}
+		}
+
+		return 0;
+	}
+	else
+	{
+		LogError(AuthTokenResult.GetError());
+		return -1;
+	}
+}
+
+}  // namespace unsync
diff --git a/Engine/Source/Programs/Unsync/Private/UnsyncCmdLogin.h b/Engine/Source/Programs/Unsync/Private/UnsyncCmdLogin.h
new file mode 100644
--- /dev/null
+++ b/Engine/Source/Programs/Unsync/Private/UnsyncCmdLogin.h
@@ -0,0 +1,20 @@
+// Copyright Epic Games, Inc. All Rights Reserved.
+
+#pragma once
+
+#include "UnsyncCommon.h"
+#include "UnsyncRemote.h"
+
+namespace unsync {
+
+struct FCmdLoginOptions
+{
+	FRemoteDesc Remote;
+	bool		bInteractive = false;
+	bool		bPrint		 = false;
+	bool		bDecode		 = false;
+};
+
+int32 CmdLogin(const FCmdLoginOptions& Options);
+
+}  // namespace unsync
diff --git a/Engine/Source/Programs/Unsync/Private/UnsyncCmdQuery.cpp b/Engine/Source/Programs/Unsync/Private/UnsyncCmdQuery.cpp
--- a/Engine/Source/Programs/Unsync/Private/UnsyncCmdQuery.cpp
+++ b/Engine/Source/Programs/Unsync/Private/UnsyncCmdQuery.cpp
@@ -6,6 +6,7 @@
 #include "UnsyncProxy.h"
 #include "UnsyncThread.h"
 #include "UnsyncUtil.h"
+#include "UnsyncAuth.h"
 
 #include <float.h>
 #include <algorithm>
@@ -13,10 +14,6 @@
 
 #include <fmt/format.h>
 
-#if UNSYNC_USE_TLS
-#	include "UnsyncAuth.h"
-#endif	// UNSYNC_USE_TLS
-
 namespace unsync {
 
 using FMirrorInfoResult = TResult<std::vector<FMirrorInfo>>;
@@ -121,13 +118,13 @@ CmdQueryMirrors(const FCmdQueryOptions& Options)
 		return A < B;
 	});
 
-	LogPrintf(ELogLevel::Info, L"[\n");
+	LogPrintf(ELogLevel::MachineReadable, L"[\n");
 
 	for (size_t I = 0; I < Mirrors.size(); ++I)
 	{
 		const FMirrorInfo& Mirror = Mirrors[I];
 
-		LogPrintf(ELogLevel::Info,
+		LogPrintf(ELogLevel::MachineReadable,
 				  L"  {\"address\":\"%hs\", \"port\":%d, \"ok\":%hs, \"ping\":%d, \"name\":\"%hs\"}%hs\n",
 				  Mirror.Address.c_str(),
 				  Mirror.Port,
@@ -137,76 +134,11 @@ CmdQueryMirrors(const FCmdQueryOptions& Options)
 				  I + 1 == Mirrors.size() ? "" : ",");
 	}
 
-	LogPrintf(ELogLevel::Info, L"]\n");
+	LogPrintf(ELogLevel::MachineReadable, L"]\n");
 
 	return 0;
 }
 
-#if UNSYNC_USE_TLS
-int32
-CmdQueryLogin(const FCmdQueryOptions& Options)
-{
-	TResult<FAuthToken> AuthToken = Authenticate(Options.Remote, 5 * 60);
-
-	if (AuthToken.IsOk())
-	{
-		{
-			FHttpConnection Connection = FHttpConnection::CreateDefaultHttps(Options.Remote);
-
-			FHttpRequest Request;
-			Request.Url			   = "/api/v1/login";
-			Request.Method		   = EHttpMethod::GET;
-			Request.BearerToken	   = AuthToken->Access;
-			FHttpResponse Response = HttpRequest(Connection, Request);
-
-			if (Response.Success())
-			{
-				UNSYNC_LOG("Login successful");
-			}
-			else
-			{
-				LogError(HttpError(Response.Code));
-				return -1;
-			}
-		}
-
-		UNSYNC_LOG(L"Getting user info from authentication server");
-		UNSYNC_LOG_INDENT;
-
-		TResult<FAuthDesc> AuthDescResult = GetAuthenticationDesc(Options.Remote);
-		if (AuthDescResult.IsError())
-		{
-			LogError(AuthDescResult.GetError());
-			return -1;
-		}
-
-		const FAuthDesc& AuthDesc = AuthDescResult.GetData();
-
-		FHttpConnection AuthConnection = FHttpConnection::CreateDefaultHttps(AuthDesc.ServerHost);
-
-		TResult<FAuthUserInfo> UserInfoResult = GetUserInfo(AuthConnection, AuthDesc, AuthToken.GetData());
-
-		if (const FAuthUserInfo* UserInfo = UserInfoResult.TryData())
-		{
-			if (UserInfo->Email.length() && UserInfo->Name.length())
-			{
-				UNSYNC_VERBOSE(L"Authenticated user: %hs (%hs)", UserInfo->Name.c_str(), UserInfo->Email.c_str());
-			}
-			else if (UserInfo->Name.length())
-			{
-				UNSYNC_VERBOSE(L"Authenticated user: %hs", UserInfo->Name.c_str());
-			}
-		}
-
-		return 0;
-	}
-	else
-	{
-		LogError(AuthToken.GetError());
-		return -1;
-	}
-}
-
 int32
 CmdQueryList(const FCmdQueryOptions& Options)
 {
@@ -256,7 +188,6 @@ CmdQueryList(const FCmdQueryOptions& Options)
 
 	return 0;
 }
-#endif	// UNSYNC_USE_TLS
 
 int32
 CmdQuery(const FCmdQueryOptions& Options)
@@ -265,16 +196,10 @@ CmdQuery(const FCmdQueryOptions& Options)
 	{
 		return CmdQueryMirrors(Options);
 	}
-#if UNSYNC_USE_TLS
-	else if (Options.Query == "login")
-	{
-		return CmdQueryLogin(Options);
-	}
 	else if (Options.Query == "list")
 	{
 		return CmdQueryList(Options);
 	}
-#endif	// UNSYNC_USE_TLS
 	else
 	{
 		UNSYNC_ERROR(L"Unknown query command");
diff --git a/Engine/Source/Programs/Unsync/Private/UnsyncCommon.h b/Engine/Source/Programs/Unsync/Private/UnsyncCommon.h
--- a/Engine/Source/Programs/Unsync/Private/UnsyncCommon.h
+++ b/Engine/Source/Programs/Unsync/Private/UnsyncCommon.h
@@ -22,10 +22,6 @@ using int64 = int64_t;
 
 using FPath = std::filesystem::path;
 
-#ifndef UNSYNC_USE_TLS
-#	define UNSYNC_USE_TLS 0
-#endif
-
 #define UNSYNC_DISALLOW_COPY_ASSIGN(T) \
 	T(const T&) = delete;              \
 	T& operator=(const T&) = delete;
diff --git a/Engine/Source/Programs/Unsync/Private/UnsyncHttp.cpp b/Engine/Source/Programs/Unsync/Private/UnsyncHttp.cpp
--- a/Engine/Source/Programs/Unsync/Private/UnsyncHttp.cpp
+++ b/Engine/Source/Programs/Unsync/Private/UnsyncHttp.cpp
@@ -530,7 +530,6 @@ FHttpConnection::Open()
 
 	if (bUseTls)
 	{
-#if UNSYNC_USE_TLS
 		FTlsClientSettings ClientSettings;
 		ClientSettings.bVerifyCertificate = bTlsVerifyCertificate;
 
@@ -554,9 +553,6 @@ FHttpConnection::Open()
 		{
 			delete TlsSocket;
 		}
-#else	// UNSYNC_USE_TLS
-		UNSYNC_ERROR(L"Unsync is not compiled with TLS support");
-#endif	// UNSYNC_USE_TLS
 	}
 	else
 	{
diff --git a/Engine/Source/Programs/Unsync/Private/UnsyncLog.cpp b/Engine/Source/Programs/Unsync/Private/UnsyncLog.cpp
--- a/Engine/Source/Programs/Unsync/Private/UnsyncLog.cpp
+++ b/Engine/Source/Programs/Unsync/Private/UnsyncLog.cpp
@@ -35,12 +35,13 @@ IsDebuggerPresent()
 }
 #endif
 
-thread_local bool	GLogVerbose		= false;
-bool				GLogVeryVerbose = false;
-bool				GBreakOnError	= IsDebuggerPresent();
-bool				GBreakOnWarning = false;
-thread_local uint32 GLogIndent		= 0;
-bool				GLogProgress	= false;
+thread_local bool	GLogVerbose			= false;
+bool				GLogVeryVerbose		= false;
+bool				GBreakOnError		= IsDebuggerPresent();
+bool				GBreakOnWarning		= false;
+thread_local uint32 GLogIndent			= 0;
+bool				GLogProgress		= false;
+bool				GLogMachineReadable = false;
 
 std::mutex GLogMutex;
 
@@ -48,6 +49,27 @@ std::atomic<uint32> GLogThreadIndexCounter;
 thread_local uint32 GLogThreadIndex = ~0u;
 
 FTimePoint GNextFlushTime = TimePointNow();
+
+static FILE*
+GetLogStream(ELogLevel LogLevel)
+{
+	if (GLogMachineReadable)
+	{
+		if (LogLevel == ELogLevel::MachineReadable)
+		{
+			return stdout;
+		}
+		else
+		{
+			return stderr;
+		}
+	}
+	else
+	{
+		return stdout;
+	}
+}
+
 static void
 LogConditionalFlush(FILE* Stream)
 {
@@ -60,6 +82,12 @@ LogConditionalFlush(FILE* Stream)
 	}
 }
 
+static void
+LogConditionalFlush(ELogLevel LogLevel)
+{
+	LogConditionalFlush(GetLogStream(LogLevel));
+}
+
 static uint32
 GetLogThreadIndex()
 {
@@ -143,7 +171,7 @@ LogProgress(const wchar_t* ItemName, uint64 Current, uint64 Total)
 	}
 	wprintf(L"@progress [%ls] %llu / %llu\n", ItemName, Current, Total);
 
-	LogConditionalFlush(stdout);
+	LogConditionalFlush(stdout); // progress is always reported to stdout
 }
 
 void
@@ -161,7 +189,7 @@ LogStatus(const wchar_t* InItemName, const wchar_t* Status)
 
 	if (InItemName)
 	{
-		LogConditionalFlush(stdout);
+		LogConditionalFlush(stdout); // status is always reported to stdout
 	}
 	else
 	{
@@ -174,6 +202,7 @@ LogFlush()
 {
 	std::lock_guard<std::mutex> LockGuard(GLogMutex);
 	fflush(stdout);
+	fflush(stderr);
 	if (GLogFile && GLogFile->Handle)
 	{
 		FILE* F = GLogFile->Handle;
@@ -217,49 +246,51 @@ LogPrintf(ELogLevel Level, const wchar_t* Str, ...)
 		}
 	}
 
-	if (Level <= MaxDisplayLevel)
+	FILE* LogStream = GetLogStream(Level);
+
+	if (Level <= MaxDisplayLevel || Level == ELogLevel::MachineReadable)
 	{
 		if (Prefix)
 		{
-			wprintf(Prefix);
+			fwprintf(LogStream, Prefix);
 		}
 
 		if (bShouldIndent)
 		{
-			wprintf(L"%*c", GLogIndent, L' ');
+			fwprintf(LogStream, L"%*c", GLogIndent, L' ');
 		}
 
 		va_list Va;
 		va_start(Va, Str);
-		vwprintf(Str, Va);
+		vfwprintf(LogStream, Str, Va);
 		va_end(Va);
 
-		LogConditionalFlush(stdout);
+		LogConditionalFlush(Level);
 	}
 
 	if (GLogFile && GLogFile->Handle)
 	{
-		FILE* F = GLogFile->Handle;
+		FILE* LogFileStream = GLogFile->Handle;
 
 		uint32 ThreadIndex = GetLogThreadIndex();
-		fwprintf(F, L"[%3d] ", ThreadIndex);
+		fwprintf(LogFileStream, L"[%3d] ", ThreadIndex);
 
 		switch (Level)
 		{
 			case ELogLevel::Error:
-				fwprintf(F, L"[ERROR] ");
+				fwprintf(LogFileStream, L"[ERROR] ");
 				break;
 			case ELogLevel::Warning:
-				fwprintf(F, L"[WARN] ");
+				fwprintf(LogFileStream, L"[WARN] ");
 				break;
 			case ELogLevel::Info:
-				fwprintf(F, L"[INFO] ");
+				fwprintf(LogFileStream, L"[INFO] ");
 				break;
 			case ELogLevel::Debug:
-				fwprintf(F, L"[DEBUG] ");
+				fwprintf(LogFileStream, L"[DEBUG] ");
 				break;
 			case ELogLevel::Trace:
-				fwprintf(F, L"[TRACE] ");
+				fwprintf(LogFileStream, L"[TRACE] ");
 				break;
 			default:
 				break;
@@ -267,12 +298,12 @@ LogPrintf(ELogLevel Level, const wchar_t* Str, ...)
 
 		va_list Va;
 		va_start(Va, Str);
-		vfwprintf(F, Str, Va);
+		vfwprintf(LogFileStream, Str, Va);
 		va_end(Va);
 
 		if (Level == ELogLevel::Error || Level == ELogLevel::Warning)
 		{
-			fflush(F);
+			fflush(LogFileStream);
 		}
 	}
 }
diff --git a/Engine/Source/Programs/Unsync/Private/UnsyncLog.h b/Engine/Source/Programs/Unsync/Private/UnsyncLog.h
--- a/Engine/Source/Programs/Unsync/Private/UnsyncLog.h
+++ b/Engine/Source/Programs/Unsync/Private/UnsyncLog.h
@@ -11,7 +11,8 @@ extern bool				   GLogVeryVerbose;
 extern bool				   GBreakOnError;
 extern bool				   GBreakOnWarning;
 extern thread_local uint32 GLogIndent;
-extern bool				   GLogProgress;
+extern bool				   GLogProgress;		 // Whether to output @progress and @status markers to stdout
+extern bool				   GLogMachineReadable;	 // Whether output is intended for other programs rather than humans
 
 struct FError;
 
@@ -63,6 +64,10 @@ enum class ELogLevel
 	Info	= 2,
 	Debug	= 3,
 	Trace	= 4,
+
+	// Special log level to signify output that's intended to be machine-readable.
+	// This is always written to stdout, while other modes may be written to stderr in some cases.
+	MachineReadable = 5,
 };
 
 void LogFlush();
diff --git a/Engine/Source/Programs/Unsync/Private/UnsyncMain.cpp b/Engine/Source/Programs/Unsync/Private/UnsyncMain.cpp
--- a/Engine/Source/Programs/Unsync/Private/UnsyncMain.cpp
+++ b/Engine/Source/Programs/Unsync/Private/UnsyncMain.cpp
@@ -1,7 +1,9 @@
 // Copyright Epic Games, Inc. All Rights Reserved.
 
+#include "UnsyncAuth.h"
 #include "UnsyncCmdDiff.h"
 #include "UnsyncCmdHash.h"
+#include "UnsyncCmdLogin.h"
 #include "UnsyncCmdMount.h"
 #include "UnsyncCmdPatch.h"
 #include "UnsyncCmdPush.h"
@@ -14,7 +16,6 @@
 #include "UnsyncTest.h"
 #include "UnsyncThread.h"
 #include "UnsyncUtil.h"
-#include "UnsyncAuth.h"
 
 UNSYNC_THIRD_PARTY_INCLUDES_START
 #if UNSYNC_PLATFORM_WINDOWS
@@ -85,6 +86,9 @@ InnerMain(int Argc, char** Argv)
 	bool					 bFullDifference	 = false;
 	bool					 bInfoFiles			 = false;
 	bool					 bNoProxySelect		 = false;
+	bool					 bInteractive		 = false;
+	bool					 bDecode			 = false;
+	bool					 bPrint				 = false;
 	int32					 CompressionLevel	 = 3;
 	uint32					 DiffBlockSize		 = uint32(4_KB);
 	uint32					 HashOrSyncBlockSize = uint32(64_KB);
@@ -98,13 +102,11 @@ InnerMain(int Argc, char** Argv)
 
 	const std::string HiddenGroupId;  // CLI11 uses an empty string group name to mark arguments that should be hidden
 
-#if UNSYNC_USE_TLS
 	auto AddTlsOptions = [&CacertFilenameUtf8, &bUseTls, &bAllowInsecureTls](CLI::App* App) {
 		App->add_option("--cacert", CacertFilenameUtf8, "Certificate authority file to use for TLS validation (.pem)");
 		App->add_flag("--tls", bUseTls, "Use TLS when connecting to remote server");
 		App->add_flag("--insecure", bAllowInsecureTls, "Skip remote server TLS certificate validation");
 	};
-#endif	// UNSYNC_USE_TLS
 
 	auto AddProxyOptions = [&RemoteAddressUtf8, &bNoProxySelect](CLI::App* App) {
 		App->add_option("--proxy, --remote, --server",
@@ -115,6 +117,8 @@ InnerMain(int Argc, char** Argv)
 					  "Skip automatic server selection and use the exact one specified by command line or environment variable");
 	};
 
+	// Configure hash
+
 	CLI::App* SubHash = Cli.add_subcommand("hash", "Generate hash manifest for a file or directory");
 	SubHash->add_option("Input", InputFilenameUtf8, "Input file or directory path")->required();
 	SubHash->add_flag("-f, --force", bForceOperation, "Force the operation even if hash is already computed for the input");
@@ -139,6 +143,8 @@ InnerMain(int Argc, char** Argv)
 		"Create a directory manifest incrementally, by updating an existing manifest if one exists (only process changed files)");
 	SubCommands.push_back(SubHash);
 
+	// Configure push
+
 	CLI::App* SubPush = Cli.add_subcommand("push", "Loads a manifest from a directory and uploads referenced blocks to the remote server");
 	SubPush->add_option("Input", InputFilenameUtf8, "Input file or directory path")->required();
 	SubPush
@@ -149,9 +155,7 @@ InnerMain(int Argc, char** Argv)
 	SubPush->add_option("--http-header-file",
 						HttpHeaderFilenameUtf8,
 						"Text file that contains any extra HTTP headers to pass to the remote server (auth tokens, etc.)");
-#if UNSYNC_USE_TLS
 	SubPush->add_flag("--insecure", bAllowInsecureTls, "Skip remote server TLS certificate validation");
-#endif
 	SubCommands.push_back(SubPush);
 
 	CLI::App* SubInfo = Cli.add_subcommand("info", "Display information about a manifest file or diff two manifests");
@@ -167,6 +171,8 @@ InnerMain(int Argc, char** Argv)
 						"Exclude filenames that contain specified words (comma separated). Filter is run after --include.");
 	SubCommands.push_back(SubInfo);
 
+	// Configure diff
+
 	CLI::App* SubDiff = Cli.add_subcommand("diff", "Compute difference required to transform BaseFile into SourceFile");
 	SubDiff->add_option("Base", BaseFilenameUtf8, "Base file name (local data)")->required();
 	SubDiff->add_option("Source", SourceFilenameUtf8, "Source file name (remote data)")->required();
@@ -175,6 +181,8 @@ InnerMain(int Argc, char** Argv)
 	SubDiff->add_option("-b, --block", DiffBlockSize, "Block size in bytes (default=4KB)");
 	SubCommands.push_back(SubDiff);
 
+	// Configure sync
+
 	CLI::App* SubSync = Cli.add_subcommand("sync", "Synchronize files, transforming target file/directory into source");
 	SubSync
 		->add_option("Source",
@@ -196,11 +204,8 @@ InnerMain(int Argc, char** Argv)
 	SubSync->add_option("--exclude",
 						ExcludeFilterArrayUtf8,
 						"Exclude filenames that contain specified words (comma separated). Filter is run after --include.");
-#if UNSYNC_USE_TLS
 	AddTlsOptions(SubSync);
-#else
-	UNSYNC_UNUSED(bUseTls);
-#endif	// UNSYNC_USE_TLS
+
 	SubSync->add_option("--http-header-file",
 						HttpHeaderFilenameUtf8,
 						"Text file that contains any extra HTTP headers to pass to the remote server (auth tokens, etc.)");
@@ -260,16 +265,28 @@ InnerMain(int Argc, char** Argv)
 	SubTest->add_option("--preset", PresetUtf8, "Test preset")->default_str(PresetUtf8);
 	SubCommands.push_back(SubTest);
 
+	// Configure query
+
 	CLI::App* SubQuery = Cli.add_subcommand("query", "Run a query command on the remote server");
-	SubQuery->add_option("QueryString", QueryStringUtf8, "Query")->required();
+	SubQuery->add_option("QueryString", QueryStringUtf8, "Query to run: mirrors, login, list")->required();
 	SubQuery->add_option("QueryArgs", QueryArgsUtf8, "Query arguments");
 	AddProxyOptions(SubQuery);
 
-#if UNSYNC_USE_TLS
 	AddTlsOptions(SubQuery);
-#endif	// UNSYNC_USE_TLS
 	SubCommands.push_back(SubQuery);
 
+	// Configure login
+
+	CLI::App* SubLogin = Cli.add_subcommand("login", "Authenticate with the remote server (acquire access and refresh tokens)");
+	SubLogin->add_flag("--interactive", bInteractive, "Allow user interaction through modal dialogs");
+	SubLogin->add_flag("--decode", bDecode, "Decode authentication token (implies --print)");
+	SubLogin->add_flag("--print", bPrint, "Print authentication token to standard output");
+	AddTlsOptions(SubLogin);
+	AddProxyOptions(SubLogin);
+	SubCommands.push_back(SubLogin);
+
+	// Configure mount
+
 	CLI::App* SubMount = Cli.add_subcommand("mount", "Mount directory manifest as a virtual file system (EXPERIMENTAL)");
 	SubMount
 		->add_option("Source",
@@ -290,6 +307,8 @@ InnerMain(int Argc, char** Argv)
 		Subcommand->add_flag("--debug", bUseDebugMode, "Enable extra debugging features, such as extra memory safety validation");
 	}
 
+	// Run the command
+
 #if UNSYNC_PLATFORM_WINDOWS
 	_setmode(_fileno(stdout), _O_TEXT);
 #endif	// UNSYNC_PLATFORM_WINDOWS
@@ -305,6 +324,11 @@ InnerMain(int Argc, char** Argv)
 	_setmode(_fileno(stdout), _O_U8TEXT);
 #endif	// UNSYNC_PLATFORM_WINDOWS
 
+	if (Cli.got_subcommand(SubQuery) || Cli.got_subcommand(SubLogin))
+	{
+		GLogMachineReadable = true;
+	}
+
 	if (GLogVeryVerbose)
 	{
 		// Force verbose mode if very-verbose flag is present
@@ -655,7 +679,8 @@ InnerMain(int Argc, char** Argv)
 		}
 	}
 
-	if (!bNoProxySelect && !Cli.got_subcommand(SubQuery)
+	if (!bNoProxySelect
+		&& Cli.got_subcommand(SubSync)
 		&& RemoteDesc.IsValid() && RemoteDesc.Protocol == EProtocolFlavor::Unsync)
 	{
 		UNSYNC_LOG(L"Selecting server using root '%hs'", RemoteDesc.HostAddress.c_str());
@@ -710,7 +735,6 @@ InnerMain(int Argc, char** Argv)
 			ScavengeRoot = FPath{};
 		}
 
-#if UNSYNC_USE_TLS
 		{
 			UNSYNC_VERBOSE(L"Attempting to authenticate");
 			UNSYNC_LOG_INDENT;
@@ -720,7 +744,6 @@ InnerMain(int Argc, char** Argv)
 				UNSYNC_VERBOSE(L"Authentication enabled");
 			}
 		}
-#endif	// UNSYNC_USE_TLS
 
 		FCmdSyncOptions SyncOptions;
 
@@ -785,6 +808,20 @@ InnerMain(int Argc, char** Argv)
 		QueryOptions.Remote = RemoteDesc;
 		return CmdQuery(QueryOptions);
 	}
+	else if (Cli.got_subcommand(SubLogin))
+	{
+		if (bDecode)
+		{
+			bPrint = true;
+		}
+
+		FCmdLoginOptions LoginOptions;
+		LoginOptions.Remote		  = RemoteDesc;
+		LoginOptions.bInteractive = bInteractive;
+		LoginOptions.bDecode	  = bDecode;
+		LoginOptions.bPrint		  = bPrint;
+		return CmdLogin(LoginOptions);
+	}
 	else if (Cli.got_subcommand(SubMount))
 	{
 		FCmdMountOptions MountOptions;
diff --git a/Engine/Source/Programs/Unsync/Private/UnsyncProxy.cpp b/Engine/Source/Programs/Unsync/Private/UnsyncProxy.cpp
--- a/Engine/Source/Programs/Unsync/Private/UnsyncProxy.cpp
+++ b/Engine/Source/Programs/Unsync/Private/UnsyncProxy.cpp
@@ -69,7 +69,6 @@ FUnsyncProtocolImpl::FUnsyncProtocolImpl(const FRemoteDesc&				RemoteDesc,
 : FRemoteProtocolBase(RemoteDesc, InRequestMap)
 , Features(InFeatures)
 {
-#if UNSYNC_USE_TLS
 	if (RemoteDesc.bTlsEnable && TlsSettings)
 	{
 		FSocketHandle RawSocketHandle = SocketConnectTcp(RemoteDesc.HostAddress.c_str(), RemoteDesc.HostPort);
@@ -86,7 +85,6 @@ FUnsyncProtocolImpl::FUnsyncProtocolImpl(const FRemoteDesc&				RemoteDesc,
 			}
 		}
 	}
-#endif	// UNSYNC_USE_TLS
 
 	if (!SocketHandle)
 	{
diff --git a/Engine/Source/Programs/Unsync/Private/UnsyncSocket.cpp b/Engine/Source/Programs/Unsync/Private/UnsyncSocket.cpp
--- a/Engine/Source/Programs/Unsync/Private/UnsyncSocket.cpp
+++ b/Engine/Source/Programs/Unsync/Private/UnsyncSocket.cpp
@@ -26,9 +26,7 @@ UNSYNC_THIRD_PARTY_INCLUDES_START
 #include <unordered_set>
 #include <limits>
 
-#if UNSYNC_USE_TLS
-#	include <tls.h>
-#endif	// UNSYNC_USE_TLS
+#include <tls.h>
 
 UNSYNC_THIRD_PARTY_INCLUDES_END
 
@@ -55,9 +53,7 @@ struct FSocketInitHelper
 		}
 #endif	// UNSYNC_PLATFORM_WINDOWS
 
-#if UNSYNC_USE_TLS
 		tls_init();
-#endif	// UNSYNC_USE_TLS
 	}
 	~FSocketInitHelper()
 	{
@@ -294,7 +290,6 @@ FSocketRaw::RecvAny(void* Data, size_t DataSize)
 	return SocketRecvAny(Handle, Data, DataSize);
 }
 
-#if UNSYNC_USE_TLS
 FSocketTls::FSocketTls(FSocketHandle InHandle, FTlsClientSettings ClientSettings) : FSocketBase(InHandle)
 {
 	Security = ESocketSecurity::Unknown;
@@ -441,7 +436,6 @@ FSocketTls::RecvAny(void* Data, size_t DataSize)
 
 	return ProcessedBytes;
 }
-#endif	// UNSYNC_USE_TLS
 
 FSocketBase::~FSocketBase()
 {
@@ -455,12 +449,10 @@ ToString(ESocketSecurity Security)
 	{
 		case ESocketSecurity::None:
 			return "None";
-#if UNSYNC_USE_TLS
 		case ESocketSecurity::TLSv1_2:
 			return "TLS 1.2";
 		case ESocketSecurity::TLSv1_3:
 			return "TLS 1.3";
-#endif	// UNSYNC_USE_TLS
 		default:
 			return "Unknown";
 	}
diff --git a/Engine/Source/Programs/Unsync/Private/UnsyncSocket.h b/Engine/Source/Programs/Unsync/Private/UnsyncSocket.h
--- a/Engine/Source/Programs/Unsync/Private/UnsyncSocket.h
+++ b/Engine/Source/Programs/Unsync/Private/UnsyncSocket.h
@@ -65,10 +65,8 @@ SocketRecvT(FSocketHandle Socket, T& Data)
 enum class ESocketSecurity
 {
 	None,
-#if UNSYNC_USE_TLS
 	TLSv1_2,
 	TLSv1_3,
-#endif	// UNSYNC_USE_TLS
 	Unknown
 };
 
@@ -107,7 +105,6 @@ struct FTlsClientSettings
 	uint64		 CacertSize			= 0;
 };
 
-#if UNSYNC_USE_TLS
 struct FSocketTls : FSocketBase
 {
 	FSocketTls(FSocketHandle InHandle, FTlsClientSettings ClientSettings);
@@ -121,7 +118,6 @@ struct FSocketTls : FSocketBase
 
 	tls* TlsCtx = {};
 };
-#endif	// UNSYNC_USE_TLS
 
 template<typename T>
 bool
diff --git a/Engine/Source/Programs/Unsync/README.md b/Engine/Source/Programs/Unsync/README.md
--- a/Engine/Source/Programs/Unsync/README.md
+++ b/Engine/Source/Programs/Unsync/README.md
@@ -216,11 +216,3 @@ dependencies and build optimized binary with debug symbols:
 > cmake -B build -S .
 > cmake --build build --config RelWithDebInfo
 ```
-
-Some optional features, such as TLS support, may be disabled at compile time to
-produce a smaller executable:
-
-```cmd
-> cmake -DUNSYNC_USE_TLS=OFF -B build -S .
-```
-
diff --git a/Engine/Source/Programs/Unsync/Unsync.Build.cs b/Engine/Source/Programs/Unsync/Unsync.Build.cs
--- a/Engine/Source/Programs/Unsync/Unsync.Build.cs
+++ b/Engine/Source/Programs/Unsync/Unsync.Build.cs
@@ -14,7 +14,6 @@ public class Unsync : ModuleRules
 
 		PrivateIncludePaths.Add(Path.Combine(ModuleDirectory, "ThirdParty"));
 
-		PrivateDefinitions.Add("UNSYNC_USE_TLS=1");
 		PrivateDefinitions.Add("UNSYNC_USE_DEBUG_HEAP=1");
 
 		if (Target.Platform == UnrealTargetPlatform.Win64)
