unsync - Add support for direct manifest download from proxy server (bypassing SMB)

* Add --login flag to sync command
* Add in-process token cache
* Add auth token when quering server features
* Use direct download path when supported by server

diff --git a/Engine/Source/Programs/Unsync/Private/UnsyncAuth.cpp b/Engine/Source/Programs/Unsync/Private/UnsyncAuth.cpp
--- a/Engine/Source/Programs/Unsync/Private/UnsyncAuth.cpp
+++ b/Engine/Source/Programs/Unsync/Private/UnsyncAuth.cpp
@@ -9,6 +9,7 @@
 #include <fmt/format.h>
 #include <ctime>
 #include <json11.hpp>
+#include <optional>
 
 #include <openssl/err.h>
 #include <openssl/evp.h>  // Base64 encoding
@@ -570,15 +571,82 @@ GenerateTokenId(const FRemoteDesc& RemoteDesc)
 	return HashToHexString(Hash);
 }
 
+// Keeps last loaded token in memory
+struct FAuthTokenCache
+{
+	std::mutex Mutex;
+
+	struct FEntry
+	{
+		FPath			Path;
+		FFileAttributes Attrib;
+		FAuthToken		Token;
+	};
+
+	// Only keep the most recent token now, but could extend to N recent tokens in the future
+	FEntry MostRecent;
+
+	void Add(const FPath& Path, const FFileAttributes& Attrib, const FAuthToken& AuthToken)
+	{
+		std::lock_guard<std::mutex> LockGuard(Mutex);
+
+		MostRecent.Path	  = Path;
+		MostRecent.Attrib = Attrib;
+		MostRecent.Token  = AuthToken;
+	}
+
+	std::optional<FAuthToken> Get(const FPath& Path, bool bCheckFileAttributes = false)
+	{
+		std::lock_guard<std::mutex> LockGuard(Mutex);
+
+		if (MostRecent.Path == Path)
+		{
+			if (bCheckFileAttributes)
+			{
+				FFileAttributes Attrib = GetFileAttrib(Path);
+				if (Attrib.Mtime != MostRecent.Attrib.Mtime || Attrib.Size != MostRecent.Attrib.Size)
+				{
+					return {};
+				}
+			}
+
+			return std::optional<FAuthToken>(MostRecent.Token);
+		}
+
+		return {};
+	}
+};
+
+static FAuthTokenCache GAuthTokenCache;
+
 bool
 SaveAuthToken(const FPath& Path, const FAuthToken& AuthToken)
 {
-	return WriteBufferToFile(Path, (const uint8*)AuthToken.Raw.data(), AuthToken.Raw.length());
+	bool bWrittenOk = WriteBufferToFile(Path, (const uint8*)AuthToken.Raw.data(), AuthToken.Raw.length());
+	if (!bWrittenOk)
+	{
+		return false;
+	}
+
+	FFileAttributes Attrib = GetFileAttrib(Path);
+	if (!Attrib.bValid)
+	{
+		return false;
+	}
+
+	GAuthTokenCache.Add(Path, Attrib, AuthToken);
+
+	return true;
 }
 
 TResult<FAuthToken>
 LoadAuthToken(const FPath& Path)
 {
+	if (std::optional<FAuthToken> CachedToken = GAuthTokenCache.Get(Path, /*bCheckFileAttributes*/ false))
+	{
+		return ResultOk(std::move(CachedToken.value()));
+	}
+
 	FBuffer FileBuffer = ReadFileToBuffer(Path);
 	if (FileBuffer.Size())
 	{
@@ -608,7 +676,7 @@ LoadAuthToken(const FPath& Path)
 			}
 		}
 
-		return ResultOk(AuthToken);
+		return ResultOk(std::move(AuthToken));
 	}
 	else
 	{
@@ -684,16 +752,33 @@ Authenticate(const FRemoteDesc& RemoteDesc, int32 RefreshThreshold)
 		TResult<FAuthToken> LoadResult = LoadAuthToken(*TokenCachePath);
 		if (FAuthToken* LoadedToken = LoadResult.TryData())
 		{
-			UNSYNC_VERBOSE2(L"Loaded cached authentication token");
 			PreviousToken = std::move(*LoadedToken);
 		}
 	}
 
+	static FHash128 LastLoggedTokenHash;
+	FHash128		TokenHash = HashBlake3String<FHash128>(PreviousToken.Raw);
+	bool			bShouldLog = false;
+	if (LastLoggedTokenHash != TokenHash)
+	{
+		bShouldLog			= true;
+		LastLoggedTokenHash = TokenHash;
+	}
+
+	if (bShouldLog && !PreviousToken.Raw.empty())
+	{
+		UNSYNC_VERBOSE(L"Loaded cached authentication token");
+	}
+
 	int64 CurrentTime	   = GetSecondsFromUnixEpoch();
 	int64 ExpiresInSeconds = PreviousToken.ExirationTime - CurrentTime;
 	if (ExpiresInSeconds > RefreshThreshold)
 	{
-		LogAuthTokenExpiration(PreviousToken);
+		if (bShouldLog)
+		{
+			LogAuthTokenExpiration(PreviousToken);
+		}
+
 		return ResultOk(PreviousToken);
 	}
 
@@ -765,7 +850,7 @@ GetAuthenticationDesc(const FRemoteDesc& RemoteDesc)
 		return AppError(L"Authentication is only implemented for UNSYNC protocol");
 	}
 
-	TResult<ProxyQuery::FHelloResponse> HelloResponseResult = ProxyQuery::Hello(RemoteDesc);
+	TResult<ProxyQuery::FHelloResponse> HelloResponseResult = ProxyQuery::Hello(RemoteDesc, /*bAnonymous*/ true);
 	if (HelloResponseResult.IsError())
 	{
 		return MoveError<FAuthDesc>(HelloResponseResult);
diff --git a/Engine/Source/Programs/Unsync/Private/UnsyncCmdLogin.cpp b/Engine/Source/Programs/Unsync/Private/UnsyncCmdLogin.cpp
--- a/Engine/Source/Programs/Unsync/Private/UnsyncCmdLogin.cpp
+++ b/Engine/Source/Programs/Unsync/Private/UnsyncCmdLogin.cpp
@@ -33,7 +33,9 @@ GetJwtPayload(std::string_view JwtData)
 int32
 CmdLogin(const FCmdLoginOptions& Options)
 {
-	TResult<FAuthToken> AuthTokenResult = Authenticate(Options.Remote, 5 * 60);
+	const int32 RefreshThreshold = Options.bForceRefresh ? INT_MAX : 5 * 60;
+
+	TResult<FAuthToken> AuthTokenResult = Authenticate(Options.Remote, RefreshThreshold);
 
 	if (AuthTokenResult.IsOk())
 	{
diff --git a/Engine/Source/Programs/Unsync/Private/UnsyncCmdLogin.h b/Engine/Source/Programs/Unsync/Private/UnsyncCmdLogin.h
--- a/Engine/Source/Programs/Unsync/Private/UnsyncCmdLogin.h
+++ b/Engine/Source/Programs/Unsync/Private/UnsyncCmdLogin.h
@@ -13,6 +13,7 @@ struct FCmdLoginOptions
 	bool		bInteractive = false;
 	bool		bPrint		 = false;
 	bool		bDecode		 = false;
+	bool		bForceRefresh = false;
 };
 
 int32 CmdLogin(const FCmdLoginOptions& Options);
diff --git a/Engine/Source/Programs/Unsync/Private/UnsyncCmdQuery.cpp b/Engine/Source/Programs/Unsync/Private/UnsyncCmdQuery.cpp
--- a/Engine/Source/Programs/Unsync/Private/UnsyncCmdQuery.cpp
+++ b/Engine/Source/Programs/Unsync/Private/UnsyncCmdQuery.cpp
@@ -169,22 +169,40 @@ CmdQueryList(const FCmdQueryOptions& Options)
 		LogError(AppError(fmt::format("JSON error: {}", JsonErrorString.c_str())));
 		return -1;
 	}
-	/*
-	if (JsonObject.is_array())
+
+	LogPrintf(ELogLevel::MachineReadable, L"%hs\n", Response.Buffer.Data());
+
+	return 0;
+}
+
+int32
+CmdQueryFile(const FCmdQueryOptions& Options)
+{
+	TResult<FBuffer> Response = ProxyQuery::DownloadFile(Options.Remote, Options.Args);
+
+	if (Response.IsOk())
 	{
-		int32 NumFiles = int32(JsonObject.array_items().size());
-		UNSYNC_LOG("Found files: %d", NumFiles);
-		for (auto& FileIt : JsonObject.array_items())
+		FBuffer& Buffer = Response.GetData();
+
+		UNSYNC_LOG(L"Downloaded file size: %llu bytes", llu(Buffer.Size()));
+
+		if (!Options.OutputPath.empty())
 		{
-			if (auto Value = FileIt["path"]; Value.is_string())
+			bool bWriteOk = WriteBufferToFile(Options.OutputPath, Buffer);
+			if (bWriteOk)
+			{
+				UNSYNC_LOG(L"Output written to file '%s'", Options.OutputPath.wstring().c_str());
+			}
+			else
 			{
-				UNSYNC_LOG(L" %hs", Value.string_value().c_str());
+				UNSYNC_ERROR(L"Failed to write output file '%s'", Options.OutputPath.wstring().c_str());
 			}
 		}
 	}
-	*/
-
-	wprintf(L"%hs\n", Response.Buffer.Data());
+	else
+	{
+		LogError(Response.GetError());
+	}
 
 	return 0;
 }
@@ -200,6 +218,10 @@ CmdQuery(const FCmdQueryOptions& Options)
 	{
 		return CmdQueryList(Options);
 	}
+	else if (Options.Query == "file")
+	{
+		return CmdQueryFile(Options);
+	}
 	else
 	{
 		UNSYNC_ERROR(L"Unknown query command");
diff --git a/Engine/Source/Programs/Unsync/Private/UnsyncCmdQuery.h b/Engine/Source/Programs/Unsync/Private/UnsyncCmdQuery.h
--- a/Engine/Source/Programs/Unsync/Private/UnsyncCmdQuery.h
+++ b/Engine/Source/Programs/Unsync/Private/UnsyncCmdQuery.h
@@ -11,6 +11,7 @@ struct FCmdQueryOptions
 {
 	std::string Query;
 	std::string Args;
+	FPath		OutputPath;
 	FRemoteDesc Remote;
 };
 
diff --git a/Engine/Source/Programs/Unsync/Private/UnsyncCmdSync.cpp b/Engine/Source/Programs/Unsync/Private/UnsyncCmdSync.cpp
--- a/Engine/Source/Programs/Unsync/Private/UnsyncCmdSync.cpp
+++ b/Engine/Source/Programs/Unsync/Private/UnsyncCmdSync.cpp
@@ -21,6 +21,8 @@ CmdSync(const FCmdSyncOptions& Options)
 		return Options.Filter ? Options.Filter->Resolve(Path) : Path;
 	};
 
+	FProxyPool ProxyPool(Options.Remote);
+
 	std::error_code ErrorCode	   = {};
 	FPath			ResolvedSource = ResolvePath(Options.Source);
 
@@ -36,25 +38,47 @@ CmdSync(const FCmdSyncOptions& Options)
 	if (!Options.Filter->SyncIncludedWords.empty())
 	{
 		UNSYNC_VERBOSE(L"Include filter: ");
+		UNSYNC_LOG_INDENT;
 		for( const std::wstring& include : Options.Filter->SyncIncludedWords)
 		{ 
-			UNSYNC_VERBOSE(L"\t%s", include.c_str());
+			UNSYNC_VERBOSE(L" %s", include.c_str());
 		}
 	}
 	if(!Options.Filter->SyncExcludedWords.empty())
 	{
 		UNSYNC_VERBOSE(L"Exclude filter: ");
+		UNSYNC_LOG_INDENT;
 		for( const std::wstring& exclude : Options.Filter->SyncExcludedWords)
 		{ 
-			UNSYNC_VERBOSE(L"\t%s", exclude.c_str());
+			UNSYNC_VERBOSE(L"%s", exclude.c_str());
 		}
 	}
 
-	const bool bSourcePathExists	 = PathExists(ResolvedSource, ErrorCode);
-	const bool bSourceIsDirectory	 = bSourcePathExists && unsync::IsDirectory(ResolvedSource);
-	const bool bSourceIsManifestHash = !bSourcePathExists && LooksLikeHash160(Options.Source.native());
+	bool bSourceFileSystemRequired = true;
+	bool bSourcePathExists		   = false;
+	bool bSourceIsDirectory		   = false;
+	bool bSourceIsManifestHash	   = LooksLikeHash160(Options.Source.native());
 
-	bool bSourceFileSystemRequired = !bSourceIsManifestHash;
+	if (ProxyPool.IsValid())
+	{
+		const FRemoteProtocolFeatures& Features = ProxyPool.GetFeatures();
+		if (Features.bFileDownload && Features.bDirectoryListing)
+		{
+			UNSYNC_VERBOSE(L"Server supports direct file access");
+			bSourceFileSystemRequired = false;
+		}
+		else if (Features.bDownloadByHash && bSourceIsManifestHash)
+		{
+			UNSYNC_VERBOSE(L"Server supports access by manifest hash");
+			bSourceFileSystemRequired = false;
+		}
+	}
+
+	if (bSourceFileSystemRequired)
+	{
+		bSourcePathExists  = PathExists(ResolvedSource, ErrorCode);
+		bSourceIsDirectory = bSourcePathExists && unsync::IsDirectory(ResolvedSource);
+	}
 
 	std::vector<FPath> ResolvedOverlays;
 
@@ -74,15 +98,22 @@ CmdSync(const FCmdSyncOptions& Options)
 			ResolvedOverlays.push_back(ResolvedEntry);
 		}
 
-		bSourceFileSystemRequired = true;
-
-		if (!bSourcePathExists || !bSourceIsDirectory)
+		if (bSourceFileSystemRequired)
 		{
-			UNSYNC_ERROR(L"Sync overlay option requires sync source to be a directory that exists on disk.");
+			if (!bSourcePathExists || !bSourceIsDirectory)
+			{
+				UNSYNC_ERROR(L"Sync overlay option requires sync source to be a directory that exists on disk.");
+				return 1;
+			}
+		}
+		
+		if (bSourceIsManifestHash)
+		{
+			UNSYNC_ERROR(L"Sync overlay option is not compatible with sync by manifest hash.");
 			return 1;
 		}
 
-		if (bSourceIsManifestHash)
+		if (!Options.SourceManifestOverride.empty())
 		{
 			UNSYNC_ERROR(L"Sync overlay option is not compatible with manifest override.");
 			return 1;
@@ -118,13 +149,17 @@ CmdSync(const FCmdSyncOptions& Options)
 
 			FSyncDirectoryOptions SyncOptions;
 
-			if (bSourceIsManifestHash || !bSourceFileSystemRequired)
+			if (bSourceFileSystemRequired)
 			{
-				SyncOptions.SourceType = ESyncSourceType::Server;
+				SyncOptions.SourceType = ESyncSourceType::FileSystem;
+			}
+			else if (bSourceIsManifestHash)
+			{
+				SyncOptions.SourceType = ESyncSourceType::ServerWithManifestHash;
 			}
 			else
 			{
-				SyncOptions.SourceType = ESyncSourceType::FileSystem;
+				SyncOptions.SourceType = ESyncSourceType::Server;
 			}
 
 			SyncOptions.Source				   = ResolvedSource;
@@ -134,7 +169,7 @@ CmdSync(const FCmdSyncOptions& Options)
 			SyncOptions.ScavengeDepth		   = Options.ScavengeDepth;
 			SyncOptions.Overlays			   = ResolvedOverlays;
 			SyncOptions.SourceManifestOverride = Options.SourceManifestOverride;
-			SyncOptions.Remote				   = &Options.Remote;
+			SyncOptions.ProxyPool			   = &ProxyPool;
 			SyncOptions.SyncFilter			   = Options.Filter;
 			SyncOptions.bCleanup			   = Options.bCleanup;
 			SyncOptions.bValidateSourceFiles   = Options.bFullSourceScan;
diff --git a/Engine/Source/Programs/Unsync/Private/UnsyncCore.cpp b/Engine/Source/Programs/Unsync/Private/UnsyncCore.cpp
--- a/Engine/Source/Programs/Unsync/Private/UnsyncCore.cpp
+++ b/Engine/Source/Programs/Unsync/Private/UnsyncCore.cpp
@@ -6,13 +6,13 @@
 #include "UnsyncCore.h"
 #include "UnsyncFile.h"
 #include "UnsyncHashTable.h"
+#include "UnsyncProgress.h"
 #include "UnsyncProxy.h"
 #include "UnsyncScan.h"
 #include "UnsyncScavenger.h"
 #include "UnsyncSerialization.h"
 #include "UnsyncThread.h"
 #include "UnsyncUtil.h"
-#include "UnsyncProgress.h"
 
 #include <condition_variable>
 #include <deque>
@@ -2346,7 +2346,7 @@ SyncFile(const FNeedList&		   NeedList,
 	uint64 NeedFromBase	  = ComputeSize(NeedList.Base);
 	UNSYNC_VERBOSE(L"Need from source %.2f MB, from base: %.2f MB", SizeMb(NeedFromSource), SizeMb(NeedFromBase));
 
-	const FileAttributes TargetFileAttributes = GetFileAttrib(TargetFilePath);
+	const FFileAttributes TargetFileAttributes = GetFileAttrib(TargetFilePath);
 
 	if (!TargetFileAttributes.bValid && NeedList.Sequence.empty())
 	{
@@ -2493,8 +2493,8 @@ SyncFile(const FNeedList&		   NeedList,
 			}
 		}
 
-		const uint64 ExpectedSourceBytes	= ComputeSize(NeedList.Source);
-		const uint64 ExpectedBaseBytes		= ComputeSize(NeedList.Base);
+		const uint64 ExpectedSourceBytes = ComputeSize(NeedList.Source);
+		const uint64 ExpectedBaseBytes	 = ComputeSize(NeedList.Base);
 
 		const uint64 ActualProcessedBytes	= BuildResult.SourceBytes + BuildResult.BaseBytes;
 		const uint64 ExpectedProcessedBytes = ExpectedSourceBytes + ExpectedBaseBytes;
@@ -2602,8 +2602,8 @@ SyncFile(const FPath& SourceFilePath, const FPath& BaseFilePath, const FPath& Ta
 std::error_code
 CopyFileIfNewer(const FPath& Source, const FPath& Target)
 {
-	FileAttributes	SourceAttr = GetFileAttrib(Source);
-	FileAttributes	TargetAttr = GetFileAttrib(Target);
+	FFileAttributes SourceAttr = GetFileAttrib(Source);
+	FFileAttributes TargetAttr = GetFileAttrib(Target);
 	std::error_code Ec;
 	if (SourceAttr.Size != TargetAttr.Size || SourceAttr.Mtime != TargetAttr.Mtime)
 	{
@@ -2914,8 +2914,89 @@ ToPath(const std::wstring_view& Str)
 #endif	// UNSYNC_PLATFORM_UNIX
 }
 
+struct FPooledProxy
+{
+	FPooledProxy(FProxyPool& InProxyPool) : ProxyPool(InProxyPool) { Proxy = ProxyPool.Alloc(); }
+	~FPooledProxy() { ProxyPool.Dealloc(std::move(Proxy)); }
+	const FProxy&			operator->() const { return *Proxy; }
+	FProxyPool&				ProxyPool;
+	std::unique_ptr<FProxy> Proxy;
+};
+
+static bool
+DownloadFileIfNewer(const FRemoteDesc& RemoteDesc, const FPath& Source, const FPath& Target)
+{
+	using FDirectoryListing		 = ProxyQuery::FDirectoryListing;
+	using FDirectoryListingEntry = ProxyQuery::FDirectoryListingEntry;
+
+	FPath SourceParent = Source.parent_path();
+	FPath SourceFileName = Source.filename().string();
+
+	std::string SourceUtf8		   = ConvertWideToUtf8(Source.wstring());
+	std::string SourceParentUtf8   = ConvertWideToUtf8(SourceParent.wstring());
+	std::string SourceFileNameUtf8 = ConvertWideToUtf8(SourceFileName.wstring());
+
+	// TODO: could have a dedicated single file stat query 
+	TResult<FDirectoryListing> DirectoryListingResult = ProxyQuery::ListDirectory(RemoteDesc, SourceParentUtf8);
+	if (DirectoryListingResult.IsError())
+	{
+		LogError(DirectoryListingResult.GetError());
+		return false;
+	}
+
+	const FDirectoryListing&	  DirectoryListing = DirectoryListingResult.GetData();
+	const FDirectoryListingEntry* SourceEntry	   = nullptr;
+	for (const FDirectoryListingEntry& Entry : DirectoryListing.Entries)
+	{
+		if (Entry.Name == SourceFileNameUtf8 && !Entry.bDirectory)
+		{
+			SourceEntry = &Entry;
+			break;
+		}
+	}
+
+	if (!SourceEntry)
+	{
+		UNSYNC_ERROR(L"Remote file '%ls' does not exist", Source.wstring().c_str());
+		return false;
+	}
+
+	FFileAttributes TargetAttr = GetFileAttrib(Target);
+	if (SourceEntry->Size != TargetAttr.Size || SourceEntry->Mtime != TargetAttr.Mtime)
+	{
+		UNSYNC_VERBOSE(L"Downloading '%ls'", Source.wstring().c_str());
+		TResult<FBuffer> DownloadResult = ProxyQuery::DownloadFile(RemoteDesc, SourceUtf8);
+		if (DownloadResult.IsError())
+		{
+			LogError(DownloadResult.GetError());
+			return false;
+		}
+
+		const FBuffer& bFileBuffer = DownloadResult.GetData();
+		if (bFileBuffer.Size() != SourceEntry->Size)
+		{
+			UNSYNC_ERROR(L"Downloaded file size mismatch. Expected %llu, actual %llu.",
+						 llu(SourceEntry->Size),
+						 llu(bFileBuffer.Size()));
+			return false;
+		}
+
+		bool bFileWritten = WriteBufferToFile(Target, bFileBuffer);
+		if (!bFileWritten)
+		{
+			UNSYNC_ERROR(L"Failed to write downloaded file file '%ls'", Target.wstring().c_str());
+			return false;
+		}
+
+		SetFileMtime(Target, SourceEntry->Mtime);
+	}
+
+	return true;
+}
+
 static bool
 LoadAndMergeSourceManifest(FDirectoryManifest& Output,
+						   FProxyPool&		   ProxyPool,
 						   const FPath&		   SourcePath,
 						   const FPath&		   TempPath,
 						   FSyncFilter*		   SyncFilter,
@@ -2924,6 +3005,11 @@ LoadAndMergeSourceManifest(FDirectoryManifest& Output,
 {
 	auto ResolvePath = [SyncFilter](const FPath& Filename) -> FPath { return SyncFilter ? SyncFilter->Resolve(Filename) : Filename; };
 
+	const FRemoteProtocolFeatures& ProxyFeatures = ProxyPool.GetFeatures();
+
+	const bool bDownloadManifestFromProxy =
+		ProxyPool.RemoteDesc.Protocol == EProtocolFlavor::Unsync && ProxyFeatures.bDirectoryListing && ProxyFeatures.bFileDownload;
+
 	FDirectoryManifest LoadedManifest;
 
 	FPath SourceManifestRoot = SourcePath / ".unsync";
@@ -2942,25 +3028,40 @@ LoadAndMergeSourceManifest(FDirectoryManifest& Output,
 	std::string SourcePathHashStr	   = BytesToHexString(SourcePathHash.Data, sizeof(SourcePathHash.Data));
 	FPath		SourceManifestTempPath = TempPath / SourcePathHashStr;
 
-	if (!PathExists(SourceManifestPath))
-	{
-		UNSYNC_ERROR(L"Source manifest '%ls' does not exist", SourceManifestPath.wstring().c_str());
-		return false;
-	}
-
 	LogGlobalStatus(L"Caching source manifest");
-
 	UNSYNC_VERBOSE(L"Caching source manifest");
+
 	UNSYNC_VERBOSE(L" Source '%ls'", SourceManifestPath.wstring().c_str());
 	UNSYNC_VERBOSE(L" Target '%ls'", SourceManifestTempPath.wstring().c_str());
-	std::error_code CopyErrorCode = CopyFileIfNewer(SourceManifestPath, SourceManifestTempPath);
-	if (CopyErrorCode)
+
+	if (bDownloadManifestFromProxy)
 	{
-		UNSYNC_LOG(L"Failed to copy manifest '%ls' to '%ls'",
-				   SourceManifestPath.wstring().c_str(),
-				   SourceManifestTempPath.wstring().c_str());
-		UNSYNC_ERROR(L"%hs (%d)", CopyErrorCode.message().c_str(), CopyErrorCode.value());
-		return false;
+		UNSYNC_LOG_INDENT;
+		bool bDownloadedOk = DownloadFileIfNewer(ProxyPool.RemoteDesc, SourceManifestPath, SourceManifestTempPath);
+		if (!bDownloadedOk)
+		{
+			UNSYNC_ERROR(L"Failed to download manifest file '%ls'", SourceManifestPath.wstring().c_str());
+			return false;
+		}
+	}
+	else
+	{
+		UNSYNC_LOG_INDENT;
+		if (!PathExists(SourceManifestPath))
+		{
+			UNSYNC_ERROR(L"Source manifest '%ls' does not exist", SourceManifestPath.wstring().c_str());
+			return false;
+		}
+
+		std::error_code CopyErrorCode = CopyFileIfNewer(SourceManifestPath, SourceManifestTempPath);
+		if (CopyErrorCode)
+		{
+			UNSYNC_LOG(L"Failed to copy manifest '%ls' to '%ls'",
+					   SourceManifestPath.wstring().c_str(),
+					   SourceManifestTempPath.wstring().c_str());
+			UNSYNC_ERROR(L"%hs (%d)", CopyErrorCode.message().c_str(), CopyErrorCode.value());
+			return false;
+		}
 	}
 
 	if (!LoadDirectoryManifest(LoadedManifest, SourcePath, SourceManifestTempPath))
@@ -3078,10 +3179,14 @@ struct FFileSyncTaskBatch
 bool  // TODO: return a TResult
 SyncDirectory(const FSyncDirectoryOptions& SyncOptions)
 {
+	FProxyPool	DummyProxyPool;
+	FProxyPool& ProxyPool = SyncOptions.ProxyPool ? *SyncOptions.ProxyPool : DummyProxyPool;
+
 	FTimePoint TimeBegin = TimePointNow();
 
 	const bool bFileSystemSource = SyncOptions.SourceType == ESyncSourceType::FileSystem;
-	const bool bServerSource	 = SyncOptions.SourceType == ESyncSourceType::Server;
+	const bool bServerSource =
+		SyncOptions.SourceType == ESyncSourceType::Server || SyncOptions.SourceType == ESyncSourceType::ServerWithManifestHash;
 
 	UNSYNC_ASSERT(bFileSystemSource || bServerSource);
 
@@ -3133,37 +3238,12 @@ SyncDirectory(const FSyncDirectoryOptions& SyncOptions)
 
 	auto ResolvePath = [SyncFilter](const FPath& Filename) -> FPath { return SyncFilter ? SyncFilter->Resolve(Filename) : Filename; };
 
-	FRemoteDesc ProxySettings = SyncOptions.Remote ? *SyncOptions.Remote : FRemoteDesc();
-	FProxyPool	ProxyPool(ProxySettings);
-
 	FDirectoryManifest SourceDirectoryManifest;
 	FPath			   SourceManifestTempPath;
 
 	const bool bCaseSensitiveTargetFileSystem = IsCaseSensitiveFileSystem(TargetTempPath);
 
-	if (bFileSystemSource || !SourceManifestOverride.empty())
-	{
-		std::vector<FPath> AllSources;
-		AllSources.push_back(SourcePath);
-		for (const FPath& OverlayPath : SyncOptions.Overlays)
-		{
-			AllSources.push_back(OverlayPath);
-		}
-
-		for (const FPath& ThisSourcePath : AllSources)
-		{
-			if (!LoadAndMergeSourceManifest(SourceDirectoryManifest,
-											ThisSourcePath,
-											TargetTempPath,
-											SyncFilter,
-											SourceManifestOverride,
-											bCaseSensitiveTargetFileSystem))
-			{
-				return false;
-			}
-		}
-	}
-	else
+	if (SyncOptions.SourceType == ESyncSourceType::ServerWithManifestHash)
 	{
 		if (!ProxyPool.IsValid())
 		{
@@ -3191,13 +3271,36 @@ SyncDirectory(const FSyncDirectoryOptions& SyncOptions)
 			return false;
 		}
 	}
+	else
+	{
+		std::vector<FPath> AllSources;
+		AllSources.push_back(SourcePath);
+		for (const FPath& OverlayPath : SyncOptions.Overlays)
+		{
+			AllSources.push_back(OverlayPath);
+		}
+
+		for (const FPath& ThisSourcePath : AllSources)
+		{
+			if (!LoadAndMergeSourceManifest(SourceDirectoryManifest,
+											ProxyPool,
+											ThisSourcePath,
+											TargetTempPath,
+											SyncFilter,
+											SourceManifestOverride,
+											bCaseSensitiveTargetFileSystem))
+			{
+				return false;
+			}
+		}
+	}
 
 	{
 		UNSYNC_VERBOSE(L"Loaded manifest properties:");
 		UNSYNC_LOG_INDENT;
 		FDirectoryManifestInfo ManifestInfo = GetManifestInfo(SourceDirectoryManifest);
 		LogManifestInfo(ELogLevel::Debug, ManifestInfo);
-		if (SyncOptions.Remote->Protocol == EProtocolFlavor::Jupiter && ManifestInfo.NumMacroBlocks == 0)
+		if (ProxyPool.RemoteDesc.Protocol == EProtocolFlavor::Jupiter && ManifestInfo.NumMacroBlocks == 0)
 		{
 			UNSYNC_ERROR(L"Manifest must contain macro blocks when using Jupiter");
 			return false;
@@ -3313,7 +3416,7 @@ SyncDirectory(const FSyncDirectoryOptions& SyncOptions)
 
 		if (bFileSystemSource && SyncOptions.bValidateSourceFiles)
 		{
-			FileAttributes SourceFileAttrib = GetFileAttrib(ResolvedSourceFilePath, &SourceAttribCache);
+			FFileAttributes SourceFileAttrib = GetFileAttrib(ResolvedSourceFilePath, &SourceAttribCache);
 
 			if (!SourceFileAttrib.bValid)
 			{
@@ -3345,7 +3448,7 @@ SyncDirectory(const FSyncDirectoryOptions& SyncOptions)
 
 		if (bSourceManifestOk)
 		{
-			FileAttributes BaseFileAttrib = GetFileAttrib(BaseFilePath, &BaseAttribCache);
+			FFileAttributes BaseFileAttrib = GetFileAttrib(BaseFilePath, &BaseAttribCache);
 
 			if (!BaseFileAttrib.bValid)
 			{
@@ -3534,13 +3637,13 @@ SyncDirectory(const FSyncDirectoryOptions& SyncOptions)
 
 	LogGlobalProgress();
 
-	if (ProxySettings.IsValid() && ProxyPool.IsValid())
+	if (ProxyPool.IsValid())
 	{
 		LogGlobalStatus(L"Connecting to server");
 		UNSYNC_VERBOSE(L"Connecting to %hs server '%hs:%d' ...",
-					   ToString(ProxySettings.Protocol),
-					   ProxySettings.HostAddress.c_str(),
-					   ProxySettings.HostPort);
+					   ToString(ProxyPool.RemoteDesc.Protocol),
+					   ProxyPool.RemoteDesc.HostAddress.c_str(),
+					   ProxyPool.RemoteDesc.HostPort);
 		UNSYNC_LOG_INDENT;
 
 		std::unique_ptr<FProxy> Proxy = ProxyPool.Alloc();
@@ -3580,7 +3683,7 @@ SyncDirectory(const FSyncDirectoryOptions& SyncOptions)
 		}
 	}
 
-	if (ProxySettings.IsValid() && !ProxyPool.IsValid())
+	if (!ProxyPool.IsValid())
 	{
 		// TODO: bail out if remote connection is required for the download,
 		// such as when downloading data purely from Jupiter.
diff --git a/Engine/Source/Programs/Unsync/Private/UnsyncCore.h b/Engine/Source/Programs/Unsync/Private/UnsyncCore.h
--- a/Engine/Source/Programs/Unsync/Private/UnsyncCore.h
+++ b/Engine/Source/Programs/Unsync/Private/UnsyncCore.h
@@ -385,20 +385,21 @@ enum class ESyncSourceType
 	Unknown,
 	FileSystem,
 	Server,
+	ServerWithManifestHash,
 };
 
 struct FSyncDirectoryOptions
 {
 	ESyncSourceType	   SourceType;
-	FPath			   Source;	  // remote data location
-	FPath			   Target;	  // output target location
-	FPath			   Base;	  // base data location, which typically is the same as sync target
-	FPath			   ScavengeRoot; // base directory where we may want to find reusable blocks
-	uint32			   ScavengeDepth = 5; // how deep to look for unsync manifests
-	std::vector<FPath> Overlays;  // extra source directories to overlay over primary (add extra files, replace existing files)
+	FPath			   Source;			   // remote data location
+	FPath			   Target;			   // output target location
+	FPath			   Base;			   // base data location, which typically is the same as sync target
+	FPath			   ScavengeRoot;	   // base directory where we may want to find reusable blocks
+	uint32			   ScavengeDepth = 5;  // how deep to look for unsync manifests
+	std::vector<FPath> Overlays;		   // extra source directories to overlay over primary (add extra files, replace existing files)
 	FPath			   SourceManifestOverride;	// force the manifest to be read from a specified file instead of source directory
 	FSyncFilter*	   SyncFilter = nullptr;	// filter callback for partial sync support
-	const FRemoteDesc* Remote	  = nullptr;	// unsync proxy server connection settings
+	FProxyPool*		   ProxyPool  = nullptr;
 	bool			   bCleanup	  = false;	// whether to cleanup any files in the target directory that are not in the source manifest file
 	bool			   bValidateSourceFiles = true;	 // whether to check that all source files declared in the manifest are present/valid
 	bool			   bValidateTargetFiles = true;	 // WARNING: turning this off is intended only for testing/profiling
diff --git a/Engine/Source/Programs/Unsync/Private/UnsyncFile.cpp b/Engine/Source/Programs/Unsync/Private/UnsyncFile.cpp
--- a/Engine/Source/Programs/Unsync/Private/UnsyncFile.cpp
+++ b/Engine/Source/Programs/Unsync/Private/UnsyncFile.cpp
@@ -572,10 +572,10 @@ FWindowsFile::FlushOne()
 	Cmd.Callback = {};
 }
 
-FileAttributes
+FFileAttributes
 GetFileAttrib(const FPath& Path, FFileAttributeCache* AttribCache)
 {
-	FileAttributes Result;
+	FFileAttributes Result;
 
 	FPath ExtendedPath = MakeExtendedAbsolutePath(Path);
 
@@ -772,10 +772,10 @@ FUnixFile::Write(const void* data, uint64 DestOffset, uint64 WriteSize)
 	return wrote_bytes;
 }
 
-FileAttributes
+FFileAttributes
 GetFileAttrib(const FPath& path, FFileAttributeCache* AttribCache)
 {
-	FileAttributes result;
+	FFileAttributes result;
 
 	if (AttribCache)
 	{
@@ -1075,7 +1075,7 @@ CreateFileAttributeCache(const FPath& Root, const FSyncFilter* SyncFilter)
 			continue;
 		}
 
-		FileAttributes Attr = {};
+		FFileAttributes Attr = {};
 
 		Attr.Mtime	= ToWindowsFileTime(Dir.last_write_time());
 		Attr.Size	= Dir.file_size();
@@ -1094,7 +1094,7 @@ CreateFileAttributeCache(const FPath& Root, const FSyncFilter* SyncFilter)
 bool
 IsDirectory(const FPath& Path)
 {
-	FileAttributes Attr = GetFileAttrib(Path);
+	FFileAttributes Attr = GetFileAttrib(Path);
 	return Attr.bValid && Attr.bDirectory;
 }
 
diff --git a/Engine/Source/Programs/Unsync/Private/UnsyncFile.h b/Engine/Source/Programs/Unsync/Private/UnsyncFile.h
--- a/Engine/Source/Programs/Unsync/Private/UnsyncFile.h
+++ b/Engine/Source/Programs/Unsync/Private/UnsyncFile.h
@@ -425,7 +425,7 @@ FBuffer ReadFileToBuffer(const FPath& Filename);
 bool	WriteBufferToFile(const FPath& Filename, const uint8* Data, uint64 Size);
 bool	WriteBufferToFile(const FPath& Filename, const FBuffer& Buffer);
 
-struct FileAttributes
+struct FFileAttributes
 {
 	uint64 Mtime	  = 0;	// Windows file time (100ns ticks since 1601-01-01T00:00:00Z)
 	uint64 Size		  = 0;
@@ -436,22 +436,22 @@ struct FileAttributes
 
 struct FFileAttributeCache
 {
-	std::unordered_map<FPath::string_type, FileAttributes> Map;
+	std::unordered_map<FPath::string_type, FFileAttributes> Map;
 
 	const bool Exists(const FPath& Path) const;
 };
 
-FileAttributes GetFileAttrib(const FPath& Path, FFileAttributeCache* AttribCache = nullptr);
-bool		   SetFileMtime(const FPath& Path, uint64 Mtime);
-bool		   SetFileReadOnly(const FPath& Path, bool ReadOnly);
-bool		   IsDirectory(const FPath& Path);
-bool		   PathExists(const FPath& Path);
-bool		   PathExists(const FPath& Path, std::error_code& OutErrorCode);
-bool		   CreateDirectories(const FPath& Path);
-bool		   FileRename(const FPath& From, const FPath& To, std::error_code& OutErrorCode);
-bool		   FileCopy(const FPath& From, const FPath& To, std::error_code& OutErrorCode);
-bool		   FileCopyOverwrite(const FPath& From, const FPath& To, std::error_code& OutErrorCode);
-bool		   FileRemove(const FPath& Path, std::error_code& OutErrorCode);
+FFileAttributes GetFileAttrib(const FPath& Path, FFileAttributeCache* AttribCache = nullptr);
+bool			SetFileMtime(const FPath& Path, uint64 Mtime);
+bool			SetFileReadOnly(const FPath& Path, bool ReadOnly);
+bool			IsDirectory(const FPath& Path);
+bool			PathExists(const FPath& Path);
+bool			PathExists(const FPath& Path, std::error_code& OutErrorCode);
+bool			CreateDirectories(const FPath& Path);
+bool			FileRename(const FPath& From, const FPath& To, std::error_code& OutErrorCode);
+bool			FileCopy(const FPath& From, const FPath& To, std::error_code& OutErrorCode);
+bool			FileCopyOverwrite(const FPath& From, const FPath& To, std::error_code& OutErrorCode);
+bool			FileRemove(const FPath& Path, std::error_code& OutErrorCode);
 
 // Returns number of bytes that can be written to the given path.
 // Returns ~0ull if the available space could not be determined.
diff --git a/Engine/Source/Programs/Unsync/Private/UnsyncHttp.cpp b/Engine/Source/Programs/Unsync/Private/UnsyncHttp.cpp
--- a/Engine/Source/Programs/Unsync/Private/UnsyncHttp.cpp
+++ b/Engine/Source/Programs/Unsync/Private/UnsyncHttp.cpp
@@ -4,6 +4,7 @@
 #include "UnsyncCore.h"
 #include "UnsyncUtil.h"
 #include "UnsyncRemote.h"
+#include "UnsyncAuth.h"
 
 #include <http_parser.h>
 #include <string.h>
@@ -239,7 +240,8 @@ HttpRequest(const FRemoteDesc& RemoteDesc,
 			EHttpMethod		   Method,
 			std::string_view   RequestUrl,
 			EHttpContentType   PayloadContentType,
-			FBufferView		   Payload)
+			FBufferView		   Payload,
+			std::string_view   BearerToken)
 {
 	FTlsClientSettings TlsSettings = RemoteDesc.GetTlsClientSettings();
 	FHttpConnection	   Connection(RemoteDesc.HostAddress, RemoteDesc.HostPort, RemoteDesc.bTlsEnable ? &TlsSettings : nullptr);
@@ -251,6 +253,7 @@ HttpRequest(const FRemoteDesc& RemoteDesc,
 	Request.Url				   = RequestUrl;
 	Request.PayloadContentType = PayloadContentType;
 	Request.Payload			   = Payload;
+	Request.BearerToken		   = BearerToken;
 
 	FHttpResponse Response = HttpRequest(Connection, Request);
 
@@ -258,9 +261,9 @@ HttpRequest(const FRemoteDesc& RemoteDesc,
 }
 
 FHttpResponse
-HttpRequest(const FRemoteDesc& RemoteDesc, EHttpMethod Method, std::string_view RequestUrl)
+HttpRequest(const FRemoteDesc& RemoteDesc, EHttpMethod Method, std::string_view RequestUrl, std::string_view BearerToken)
 {
-	return HttpRequest(RemoteDesc, Method, RequestUrl, EHttpContentType::Unknown, {});
+	return HttpRequest(RemoteDesc, Method, RequestUrl, EHttpContentType::Unknown, /*Payload*/ {}, BearerToken);
 }
 
 bool
diff --git a/Engine/Source/Programs/Unsync/Private/UnsyncHttp.h b/Engine/Source/Programs/Unsync/Private/UnsyncHttp.h
--- a/Engine/Source/Programs/Unsync/Private/UnsyncHttp.h
+++ b/Engine/Source/Programs/Unsync/Private/UnsyncHttp.h
@@ -110,7 +110,8 @@ HttpRequest(FHttpConnection& Connection,
 			std::string_view Url,
 			EHttpContentType ContentType,
 			FBufferView		 Payload,
-			std::string_view CustomHeaders = {})
+			std::string_view CustomHeaders = {},
+			std::string_view BearerToken = {})
 {
 	FHttpRequest Request;
 
@@ -119,25 +120,35 @@ HttpRequest(FHttpConnection& Connection,
 	Request.PayloadContentType = ContentType;
 	Request.Payload			   = Payload;
 	Request.CustomHeaders	   = CustomHeaders;
+	Request.BearerToken		   = BearerToken;
 
 	return HttpRequest(Connection, Request);
 }
 
 inline FHttpResponse
-HttpRequest(FHttpConnection& Connection, EHttpMethod Method, std::string_view Url, std::string_view CustomHeaders = {})
+HttpRequest(FHttpConnection& Connection,
+			EHttpMethod		 Method,
+			std::string_view Url,
+			std::string_view CustomHeaders = {},
+			std::string_view BearerToken   = {})
 {
 	FHttpRequest Request;
 	Request.Method		  = Method;
 	Request.Url			  = Url;
 	Request.CustomHeaders = CustomHeaders;
+	Request.BearerToken	  = BearerToken;
 	return HttpRequest(Connection, Request);
 }
 
 FHttpResponse
-HttpRequest(const FRemoteDesc& RemoteDesc, EHttpMethod Method, std::string_view RequestUrl);
-
-FHttpResponse
-HttpRequest(const FRemoteDesc& RemoteDesc, EHttpMethod Method, std::string_view RequestUrl, EHttpContentType PayloadContentType, FBufferView Payload);
+HttpRequest(const FRemoteDesc& RemoteDesc, EHttpMethod Method, std::string_view RequestUrl, std::string_view BearerToken = {});
+
+FHttpResponse HttpRequest(const FRemoteDesc& RemoteDesc,
+						  EHttpMethod		 Method,
+						  std::string_view	 RequestUrl,
+						  EHttpContentType	 PayloadContentType,
+						  FBufferView		 Payload,
+						  std::string_view	 BearerToken = {});
 
 // Pipelined HTTP request API
 
diff --git a/Engine/Source/Programs/Unsync/Private/UnsyncMain.cpp b/Engine/Source/Programs/Unsync/Private/UnsyncMain.cpp
--- a/Engine/Source/Programs/Unsync/Private/UnsyncMain.cpp
+++ b/Engine/Source/Programs/Unsync/Private/UnsyncMain.cpp
@@ -89,6 +89,8 @@ InnerMain(int Argc, char** Argv)
 	bool					 bInteractive		 = false;
 	bool					 bDecode			 = false;
 	bool					 bPrint				 = false;
+	bool					 bShouldLogin		 = false;
+	bool					 bForceRefreshAuthentication = false;
 	int32					 CompressionLevel	 = 3;
 	uint32					 DiffBlockSize		 = uint32(4_KB);
 	uint32					 HashOrSyncBlockSize = uint32(64_KB);
@@ -250,8 +252,8 @@ InnerMain(int Argc, char** Argv)
 	SubSync->add_flag("--no-output-validation", bNoOutputValidation, "Skip final patched file block hash validation (DANGEROUS)");
 	SubSync->add_flag("--no-space-validation", bNoSpaceValidation, "Skip checking available disk space before sync (DANGEROUS)");
 	SubSync->add_option("-b, --block", HashOrSyncBlockSize, "Block size in bytes (default=64KB)");
-
 	SubSync->add_option("--scavenge", ScavengeRootUtf8, "Search for unsync manifests and reusable blocks in this directory (EXPERIMENTAL)");
+	SubSync->add_flag("--login", bShouldLogin, "Use user authentication when accessing unsync server");
 
 	SubCommands.push_back(SubSync);
 
@@ -270,6 +272,7 @@ InnerMain(int Argc, char** Argv)
 	CLI::App* SubQuery = Cli.add_subcommand("query", "Run a query command on the remote server");
 	SubQuery->add_option("QueryString", QueryStringUtf8, "Query to run: mirrors, login, list")->required();
 	SubQuery->add_option("QueryArgs", QueryArgsUtf8, "Query arguments");
+	SubQuery->add_option("-o", OutputFilenameUtf8, "Output file name");
 	AddProxyOptions(SubQuery);
 
 	AddTlsOptions(SubQuery);
@@ -281,6 +284,7 @@ InnerMain(int Argc, char** Argv)
 	SubLogin->add_flag("--interactive", bInteractive, "Allow user interaction through modal dialogs");
 	SubLogin->add_flag("--decode", bDecode, "Decode authentication token (implies --print)");
 	SubLogin->add_flag("--print", bPrint, "Print authentication token to standard output");
+	SubLogin->add_flag("--refresh", bForceRefreshAuthentication, "Force authentication refresh even if access token has not yet expired");
 	AddTlsOptions(SubLogin);
 	AddProxyOptions(SubLogin);
 	SubCommands.push_back(SubLogin);
@@ -483,7 +487,7 @@ InnerMain(int Argc, char** Argv)
 	}
 
 	GMaxThreads = std::max(1u, GMaxThreads);
-	UNSYNC_VERBOSE(L"Using threads: %d", GMaxThreads);
+	UNSYNC_VERBOSE2(L"Using threads: %d", GMaxThreads);
 	FConcurrencyPolicyScope ConcurrencyLimitScope(GMaxThreads);
 
 	if (Cli.got_subcommand(SubHash) || Cli.got_subcommand(SubSync))
@@ -679,6 +683,8 @@ InnerMain(int Argc, char** Argv)
 		}
 	}
 
+	FRemoteDesc RootRemoteDesc = RemoteDesc;
+
 	if (!bNoProxySelect
 		&& Cli.got_subcommand(SubSync)
 		&& RemoteDesc.IsValid() && RemoteDesc.Protocol == EProtocolFlavor::Unsync)
@@ -735,6 +741,7 @@ InnerMain(int Argc, char** Argv)
 			ScavengeRoot = FPath{};
 		}
 
+		if (bShouldLogin)
 		{
 			UNSYNC_VERBOSE(L"Attempting to authenticate");
 			UNSYNC_LOG_INDENT;
@@ -803,9 +810,10 @@ InnerMain(int Argc, char** Argv)
 	else if (Cli.got_subcommand(SubQuery))
 	{
 		FCmdQueryOptions QueryOptions;
-		QueryOptions.Query	= QueryStringUtf8;
-		QueryOptions.Args	= QueryArgsUtf8;
-		QueryOptions.Remote = RemoteDesc;
+		QueryOptions.Query		= QueryStringUtf8;
+		QueryOptions.Args		= QueryArgsUtf8;
+		QueryOptions.Remote		= RemoteDesc;
+		QueryOptions.OutputPath = OutputFilename;
 		return CmdQuery(QueryOptions);
 	}
 	else if (Cli.got_subcommand(SubLogin))
@@ -816,10 +824,11 @@ InnerMain(int Argc, char** Argv)
 		}
 
 		FCmdLoginOptions LoginOptions;
-		LoginOptions.Remote		  = RemoteDesc;
-		LoginOptions.bInteractive = bInteractive;
-		LoginOptions.bDecode	  = bDecode;
-		LoginOptions.bPrint		  = bPrint;
+		LoginOptions.Remote		   = RemoteDesc;
+		LoginOptions.bInteractive  = bInteractive;
+		LoginOptions.bDecode	   = bDecode;
+		LoginOptions.bPrint		   = bPrint;
+		LoginOptions.bForceRefresh = bForceRefreshAuthentication;
 		return CmdLogin(LoginOptions);
 	}
 	else if (Cli.got_subcommand(SubMount))
diff --git a/Engine/Source/Programs/Unsync/Private/UnsyncProxy.cpp b/Engine/Source/Programs/Unsync/Private/UnsyncProxy.cpp
--- a/Engine/Source/Programs/Unsync/Private/UnsyncProxy.cpp
+++ b/Engine/Source/Programs/Unsync/Private/UnsyncProxy.cpp
@@ -378,10 +378,22 @@ FUnsyncProtocolImpl::GetSocketSecurity() const
 }
 
 namespace ProxyQuery {
-TResult<FHelloResponse> Hello(const FRemoteDesc& RemoteDesc)
+TResult<FHelloResponse> Hello(const FRemoteDesc& RemoteDesc, bool bAnonymous)
 {
 	const char* Url = "/api/v1/hello";
-	FHttpResponse Response = HttpRequest(RemoteDesc, EHttpMethod::GET, Url);
+
+	std::string BearerToken;
+	if (RemoteDesc.bAuthenticationRequired && !bAnonymous)
+	{
+		TResult<FAuthToken> AuthTokenResult = Authenticate(RemoteDesc, 15 * 60);
+		if (AuthTokenResult.IsOk())
+		{
+			BearerToken = std::move(AuthTokenResult.GetData().Access);
+		}
+	}
+	
+	FHttpResponse Response = HttpRequest(RemoteDesc, EHttpMethod::GET, Url, BearerToken);
+
 	if (!Response.Success())
 	{
 		UNSYNC_ERROR(L"Failed to establish connection to UNSYNC server. Error code: %d.", Response.Code);
@@ -459,12 +471,111 @@ TResult<FHelloResponse> Hello(const FRemoteDesc& RemoteDesc)
 				{
 					Result.Features.bAuthentication = true;
 				}
+				else if (Elem.string_value() == "list")
+				{
+					Result.Features.bDirectoryListing = true;
+				}
+				else if (Elem.string_value() == "file")
+				{
+					Result.Features.bFileDownload = true;
+				}
 			}
 		}
 	}
 
 	return ResultOk(std::move(Result));
 }
+
+TResult<FDirectoryListing> ListDirectory(const FRemoteDesc& Remote, const std::string& Path)
+{
+	FDirectoryListing Result;
+
+	TResult<FAuthToken> AuthToken = Authenticate(Remote, 5 * 60);
+	if (!AuthToken.IsOk())
+	{
+		return MoveError<FDirectoryListing>(AuthToken);
+	}
+
+	FHttpConnection Connection = FHttpConnection::CreateDefaultHttps(Remote);
+
+	std::string Url = fmt::format("/api/v1/list?{}", Path);
+
+	FHttpRequest Request;
+	Request.Url			= Url;
+	Request.Method		= EHttpMethod::GET;
+	Request.BearerToken = AuthToken->Access;
+
+	FHttpResponse Response = HttpRequest(Connection, Request);
+
+	Response.Buffer.PushBack(0);
+
+	std::string	 JsonErrorString;
+	json11::Json JsonObject = json11::Json::parse((const char*)Response.Buffer.Data(), JsonErrorString);
+
+	if (!JsonErrorString.empty())
+	{
+		return AppError(fmt::format("JSON error: {}", JsonErrorString.c_str()));
+	}
+
+	const json11::Json& EntriesObject = JsonObject["entries"];
+
+	if (EntriesObject.is_array())
+	{
+		for (const auto& Elem : EntriesObject.array_items())
+		{
+			FDirectoryListingEntry Entry;
+			for (const auto& Field : Elem.object_items())
+			{
+				if (Field.first == "name" && Field.second.is_string())
+				{
+					Entry.Name = Field.second.string_value();
+				}
+				else if (Field.first == "is_directory" && Field.second.is_bool())
+				{
+					Entry.bDirectory = Field.second.bool_value();
+				}
+				else if (Field.first == "mtime" && Field.second.is_number())
+				{
+					Entry.Mtime = uint64(Field.second.number_value());
+				}
+				else if (Field.first == "size" && Field.second.is_number())
+				{
+					Entry.Size = uint64(Field.second.number_value());
+				}
+			}
+			Result.Entries.push_back(Entry);
+		}
+	}
+
+	return ResultOk(std::move(Result));
+}
+
+TResult<FBuffer> DownloadFile(const FRemoteDesc& Remote, const std::string& Path)
+{
+	TResult<FAuthToken> AuthToken = Authenticate(Remote, 5 * 60);
+	if (!AuthToken.IsOk())
+	{
+		return MoveError<FBuffer>(AuthToken);
+	}
+
+	FHttpConnection Connection = FHttpConnection::CreateDefaultHttps(Remote);
+
+	std::string Url = fmt::format("/api/v1/file?{}", Path);
+
+	FHttpRequest Request;
+	Request.Url			   = Url;
+	Request.Method		   = EHttpMethod::GET;
+	Request.BearerToken	   = AuthToken->Access;
+	FHttpResponse Response = HttpRequest(Connection, Request);
+
+	if (!Response.Success())
+	{
+		return HttpError(Response.Code);
+	}
+
+	return ResultOk(std::move(Response.Buffer));
+}
+
 }  // namespace ProxyQuery
 
 void
@@ -624,17 +735,27 @@ FBlockRequestMap::GetMacroBlockRequest(const FGenericHash& BlockHash) const
 	return Result;
 }
 
+FProxyPool::FProxyPool()
+: FProxyPool(FRemoteDesc())
+{
+}
+
 FProxyPool::FProxyPool(const FRemoteDesc& InRemoteDesc)
 : ParallelDownloadSemaphore(InRemoteDesc.MaxConnections)
 , RemoteDesc(InRemoteDesc)
 , bValid(InRemoteDesc.IsValid())
 {
-	if (bValid && RemoteDesc.Protocol == EProtocolFlavor::Unsync)
+	if (!bValid)
+	{
+		return;
+	}
+
+	if (RemoteDesc.Protocol == EProtocolFlavor::Unsync)
 	{
 		UNSYNC_VERBOSE(L"Connecting to %hs server '%hs:%d' ...",
-			ToString(RemoteDesc.Protocol),
-			RemoteDesc.HostAddress.c_str(),
-			RemoteDesc.HostPort);
+					   ToString(RemoteDesc.Protocol),
+					   RemoteDesc.HostAddress.c_str(),
+					   RemoteDesc.HostPort);
 
 		TResult<ProxyQuery::FHelloResponse> Response = ProxyQuery::Hello(RemoteDesc);
 
@@ -645,17 +766,22 @@ FProxyPool::FProxyPool(const FRemoteDesc& InRemoteDesc)
 		else
 		{
 			const ProxyQuery::FHelloResponse& Data = Response.GetData();
-			UNSYNC_VERBOSE(L"Connection established. Server name: %hs, version: %hs, git: %hs.", 
-				Data.Name.empty() ? "unknown" : Data.Name.c_str(),
-				Data.VersionNumber.empty() ? "unknown" : Data.VersionNumber.c_str(),
-				Data.VersionGit.empty() ? "unknown" : Data.VersionGit.c_str());
+			UNSYNC_VERBOSE(L"Connection established. Server name: %hs, version: %hs, git: %hs.",
+						   Data.Name.empty() ? "unknown" : Data.Name.c_str(),
+						   Data.VersionNumber.empty() ? "unknown" : Data.VersionNumber.c_str(),
+						   Data.VersionGit.empty() ? "unknown" : Data.VersionGit.c_str());
 
-			Features = Data.Features;
+			Features  = Data.Features;
 			SessionId = Data.SessionId;
 		}
 
 		bValid = Response.IsOk();
 	}
+	else if (RemoteDesc.Protocol == EProtocolFlavor::Jupiter)
+	{
+		Features.bAuthentication = true;
+		Features.bDownloadByHash = true;
+	}
 }
 
 std::unique_ptr<FProxy>
diff --git a/Engine/Source/Programs/Unsync/Private/UnsyncProxy.h b/Engine/Source/Programs/Unsync/Private/UnsyncProxy.h
--- a/Engine/Source/Programs/Unsync/Private/UnsyncProxy.h
+++ b/Engine/Source/Programs/Unsync/Private/UnsyncProxy.h
@@ -85,9 +85,12 @@ private:
 
 struct FRemoteProtocolFeatures
 {
-	bool bTelemetry = false;
-	bool bMirrors = false;
-	bool bAuthentication = false;
+	bool bTelemetry		   = false;
+	bool bMirrors		   = false;
+	bool bAuthentication   = false;
+	bool bDirectoryListing = false;
+	bool bFileDownload	   = false;
+	bool bDownloadByHash   = false;
 };
 
 struct FTelemetryEventSyncComplete
@@ -149,6 +152,7 @@ private:
 class FProxyPool
 {
 public:
+	FProxyPool();
 	FProxyPool(const FRemoteDesc& InRemoteDesc);
 
 	std::unique_ptr<FProxy> Alloc();
@@ -181,8 +185,7 @@ private:
 	std::mutex Mutex;
 };
 
-namespace ProxyQuery
-{
+namespace ProxyQuery {
 
 struct FHelloResponse
 {
@@ -199,18 +202,25 @@ struct FHelloResponse
 	std::vector<std::string> FeatureNames;
 	FRemoteProtocolFeatures	 Features;
 
-	bool SupportsAuthentication() const
-	{
-		return Features.bAuthentication && !AuthServerUri.empty() && !AuthClientId.empty();
-	}
+	bool SupportsAuthentication() const { return Features.bAuthentication && !AuthServerUri.empty() && !AuthClientId.empty(); }
+};
+TResult<FHelloResponse> Hello(const FRemoteDesc& RemoteDesc, bool bAnonymous = false);
+
+struct FDirectoryListingEntry
+{
+	std::string Name;
+	uint64		Mtime	   = 0;
+	uint64		Size	   = 0;
+	bool		bDirectory = false;
 };
-TResult<FHelloResponse> Hello(const FRemoteDesc& RemoteDesc);
 
-struct FLoginResponse
+struct FDirectoryListing
 {
-	std::shared_ptr<FBuffer> Authentication;
+	std::vector<FDirectoryListingEntry> Entries;
 };
-TResult<FLoginResponse> QueryLogin(const FRemoteDesc& RemoteDesc);
+
+TResult<FDirectoryListing> ListDirectory(const FRemoteDesc& Remote, const std::string& Path);
+TResult<FBuffer>		   DownloadFile(const FRemoteDesc& Remote, const std::string& Path);
 
 } 
 
diff --git a/Engine/Source/Programs/Unsync/Private/UnsyncScavenger.cpp b/Engine/Source/Programs/Unsync/Private/UnsyncScavenger.cpp
--- a/Engine/Source/Programs/Unsync/Private/UnsyncScavenger.cpp
+++ b/Engine/Source/Programs/Unsync/Private/UnsyncScavenger.cpp
@@ -222,7 +222,7 @@ FScavengeDatabase::IsSourceValid(FScavengeBlockSource SourceId) const
 
 	const FFileManifest& FileManifest = FileIt->second;
 
-	FileAttributes Attrib = GetFileAttrib(FullFilePath);
+	FFileAttributes Attrib = GetFileAttrib(FullFilePath);
 
 	return Attrib.bValid && Attrib.Mtime == FileManifest.Mtime && Attrib.Size == FileManifest.Size;
 }
