unsync - Add quick login option

* Skip remote token validation in --quick mode, just check expiry locally

diff --git a/Engine/Source/Programs/Unsync/Private/UnsyncCmdLogin.cpp b/Engine/Source/Programs/Unsync/Private/UnsyncCmdLogin.cpp
--- a/Engine/Source/Programs/Unsync/Private/UnsyncCmdLogin.cpp
+++ b/Engine/Source/Programs/Unsync/Private/UnsyncCmdLogin.cpp
@@ -39,6 +39,7 @@ CmdLogin(const FCmdLoginOptions& Options)
 
 	if (AuthTokenResult.IsOk())
 	{
+		if (!Options.bQuick)
 		{
 			FHttpConnection Connection = FHttpConnection::CreateDefaultHttps(Options.Remote);
 
diff --git a/Engine/Source/Programs/Unsync/Private/UnsyncCmdLogin.h b/Engine/Source/Programs/Unsync/Private/UnsyncCmdLogin.h
--- a/Engine/Source/Programs/Unsync/Private/UnsyncCmdLogin.h
+++ b/Engine/Source/Programs/Unsync/Private/UnsyncCmdLogin.h
@@ -10,10 +10,11 @@ namespace unsync {
 struct FCmdLoginOptions
 {
 	FRemoteDesc Remote;
-	bool		bInteractive = false;
-	bool		bPrint		 = false;
-	bool		bDecode		 = false;
+	bool		bInteractive  = false;	// TODO
+	bool		bPrint		  = false;
+	bool		bDecode		  = false;
 	bool		bForceRefresh = false;
+	bool		bQuick		  = false;
 };
 
 int32 CmdLogin(const FCmdLoginOptions& Options);
diff --git a/Engine/Source/Programs/Unsync/Private/UnsyncCmdQuery.cpp b/Engine/Source/Programs/Unsync/Private/UnsyncCmdQuery.cpp
--- a/Engine/Source/Programs/Unsync/Private/UnsyncCmdQuery.cpp
+++ b/Engine/Source/Programs/Unsync/Private/UnsyncCmdQuery.cpp
@@ -180,6 +180,8 @@ CmdQueryFile(const FCmdQueryOptions& Options)
 {
 	TResult<FBuffer> Response = ProxyQuery::DownloadFile(Options.Remote, Options.Args);
 
+	UNSYNC_LOG(L"Downloading file: '%hs'", Options.Args.c_str());
+
 	if (Response.IsOk())
 	{
 		FBuffer& Buffer = Response.GetData();
diff --git a/Engine/Source/Programs/Unsync/Private/UnsyncCore.cpp b/Engine/Source/Programs/Unsync/Private/UnsyncCore.cpp
--- a/Engine/Source/Programs/Unsync/Private/UnsyncCore.cpp
+++ b/Engine/Source/Programs/Unsync/Private/UnsyncCore.cpp
@@ -2972,19 +2972,19 @@ DownloadFileIfNewer(const FRemoteDesc& RemoteDesc, const FPath& Source, const FP
 			return false;
 		}
 
-		const FBuffer& bFileBuffer = DownloadResult.GetData();
-		if (bFileBuffer.Size() != SourceEntry->Size)
+		const FBuffer& FileBuffer = DownloadResult.GetData();
+		if (FileBuffer.Size() != SourceEntry->Size)
 		{
 			UNSYNC_ERROR(L"Downloaded file size mismatch. Expected %llu, actual %llu.",
 						 llu(SourceEntry->Size),
-						 llu(bFileBuffer.Size()));
+						 llu(FileBuffer.Size()));
 			return false;
 		}
 
-		bool bFileWritten = WriteBufferToFile(Target, bFileBuffer);
+		bool bFileWritten = WriteBufferToFile(Target, FileBuffer);
 		if (!bFileWritten)
 		{
-			UNSYNC_ERROR(L"Failed to write downloaded file file '%ls'", Target.wstring().c_str());
+			UNSYNC_ERROR(L"Failed to write downloaded file '%ls'", Target.wstring().c_str());
 			return false;
 		}
 
diff --git a/Engine/Source/Programs/Unsync/Private/UnsyncMain.cpp b/Engine/Source/Programs/Unsync/Private/UnsyncMain.cpp
--- a/Engine/Source/Programs/Unsync/Private/UnsyncMain.cpp
+++ b/Engine/Source/Programs/Unsync/Private/UnsyncMain.cpp
@@ -90,7 +90,8 @@
 	bool					 bDecode			 = false;
 	bool					 bPrint				 = false;
 	bool					 bShouldLogin		 = false;
-	bool					 bForceRefreshAuthentication = false;
+	bool					 bQuickLogin		 = false;
+	bool					 bForceRefreshAuth	 = false;
 	int32					 CompressionLevel	 = 3;
 	uint32					 DiffBlockSize		 = uint32(4_KB);
 	uint32					 HashOrSyncBlockSize = uint32(64_KB);
@@ -284,7 +285,8 @@
 	SubLogin->add_flag("--interactive", bInteractive, "Allow user interaction through modal dialogs");
 	SubLogin->add_flag("--decode", bDecode, "Decode authentication token (implies --print)");
 	SubLogin->add_flag("--print", bPrint, "Print authentication token to standard output");
-	SubLogin->add_flag("--refresh", bForceRefreshAuthentication, "Force authentication refresh even if access token has not yet expired");
+	SubLogin->add_flag("--refresh", bForceRefreshAuth, "Force authentication refresh even if access token has not yet expired");
+	SubLogin->add_flag("--quick", bQuickLogin, "Skip token validation using remote server (fast path when cached acess token is expected to be valid)");
 	AddTlsOptions(SubLogin);
 	AddProxyOptions(SubLogin);
 	SubCommands.push_back(SubLogin);
@@ -328,11 +330,71 @@
 	_setmode(_fileno(stdout), _O_U8TEXT);
 #endif	// UNSYNC_PLATFORM_WINDOWS
 
+	UNSYNC_VERBOSE(L"UNSYNC %hs", GetVersionString().c_str());
+
+	if (bUseDebugMode)
+	{
+		UNSYNC_LOG(L"*** Debug mode enabled ***");
+	}
+
+	UnsyncMallocInit(bUseDebugMode ? EMallocType::Debug : EMallocType::Default);
+
+	// Configure default output mehtod based on subcommand.
+	// In machine-readable mode, all verbose logging is directed to stderr.
+
 	if (Cli.got_subcommand(SubQuery) || Cli.got_subcommand(SubLogin))
 	{
 		GLogMachineReadable = true;
 	}
 
+	// Augment configuration based on environment variables if corresponding command line arguments are missing.
+
+	if (const char* EnvCleanupExclude = getenv("UNSYNC_CLEANUP_EXCLUDE"))
+	{
+		UNSYNC_VERBOSE(L"Using UNSYNC_CLEANUP_EXCLUDE environment: '%hs'", EnvCleanupExclude);
+		CleanupExcludeFilterArrayUtf8.push_back(EnvCleanupExclude);
+	}
+
+	if (PreferredDfsUtf8.empty())
+	{
+		const char* EnvDfs = getenv("UNSYNC_DFS");
+		if (EnvDfs)
+		{
+			UNSYNC_VERBOSE(L"Using UNSYNC_DFS environment: '%hs'", EnvDfs);
+			PreferredDfsUtf8 = std::string(EnvDfs);
+		}
+	}
+
+	if (RemoteAddressUtf8.empty())
+	{
+		const char* EnvProxy = getenv("UNSYNC_PROXY");
+		if (EnvProxy)
+		{
+			UNSYNC_VERBOSE(L"Using UNSYNC_PROXY environment: '%hs'", EnvProxy);
+			RemoteAddressUtf8 = std::string(EnvProxy);
+		}
+	}
+
+	if (CacertFilenameUtf8.empty())
+	{
+		const char* EnvCacert = getenv("UNSYNC_CACERT");
+		if (EnvCacert)
+		{
+			UNSYNC_VERBOSE(L"Using UNSYNC_CACERT environment: '%hs'", EnvCacert);
+			CacertFilenameUtf8 = std::string(EnvCacert);
+		}
+	}
+
+	if (HttpHeaderFilenameUtf8.empty())
+	{
+		const char* EnvHttpHeaderFile = getenv("UNSYNC_HTTP_HEADER_FILE");
+		if (EnvHttpHeaderFile)
+		{
+			UNSYNC_VERBOSE(L"Using UNSYNC_HTTP_HEADER_FILE environment: '%hs'", EnvHttpHeaderFile);
+			HttpHeaderFilenameUtf8 = std::string(EnvHttpHeaderFile);
+		}
+	}
+
 	if (GLogVeryVerbose)
 	{
 		// Force verbose mode if very-verbose flag is present
@@ -360,13 +422,6 @@
 			L"full binary difference of local files even if timestamps and sizes match.");
 	}
 
-	if (bUseDebugMode)
-	{
-		UNSYNC_LOG(L"*** Debug mode enabled ***");
-	}
-
-	UnsyncMallocInit(bUseDebugMode ? EMallocType::Debug : EMallocType::Default);
-
 	EWeakHashAlgorithmID DefaultWeakHasher = EWeakHashAlgorithmID::BuzHash;
 	if (WeakHashUtf8 == "naive")
 	{
@@ -465,8 +520,6 @@
 	FPath ScavengeRoot			 = NormalizeFilenameUtf8(ScavengeRootUtf8);
 	FPath SourceManifestFilename = NormalizeFilenameUtf8(SourceManifestFilenameUtf8);
 
-	UNSYNC_VERBOSE(L"UNSYNC %hs", GetVersionString().c_str());
-
 	if (GLogVeryVerbose)
 	{
 		UNSYNC_VERBOSE(L"Very verbose logging is enabled");
@@ -532,57 +585,11 @@
 		SyncFilter.IncludeInSync(ConvertUtf8ToWide(Str));
 	}
 
-	if (const char* EnvCleanupExclude = getenv("UNSYNC_CLEANUP_EXCLUDE"))
-	{
-		UNSYNC_VERBOSE(L"Using UNSYNC_CLEANUP_EXCLUDE environment: '%hs'", EnvCleanupExclude);
-		CleanupExcludeFilterArrayUtf8.push_back(EnvCleanupExclude);
-	}
-
 	for (const std::string& Str : CleanupExcludeFilterArrayUtf8)
 	{
 		SyncFilter.ExcludeFromCleanup(ConvertUtf8ToWide(Str));
 	}
 
-	if (PreferredDfsUtf8.empty())
-	{
-		const char* EnvDfs = getenv("UNSYNC_DFS");
-		if (EnvDfs)
-		{
-			UNSYNC_VERBOSE(L"Using UNSYNC_DFS environment: '%hs'", EnvDfs);
-			PreferredDfsUtf8 = std::string(EnvDfs);
-		}
-	}
-
-	if (RemoteAddressUtf8.empty() && Cli.got_subcommand(SubSync))
-	{
-		const char* EnvProxy = getenv("UNSYNC_PROXY");
-		if (EnvProxy)
-		{
-			UNSYNC_VERBOSE(L"Using UNSYNC_PROXY environment: '%hs'", EnvProxy);
-			RemoteAddressUtf8 = std::string(EnvProxy);
-		}
-	}
-
-	if (CacertFilenameUtf8.empty())
-	{
-		const char* EnvCacert = getenv("UNSYNC_CACERT");
-		if (EnvCacert)
-		{
-			UNSYNC_VERBOSE(L"Using UNSYNC_CACERT environment: '%hs'", EnvCacert);
-			CacertFilenameUtf8 = std::string(EnvCacert);
-		}
-	}
-
-	if (HttpHeaderFilenameUtf8.empty())
-	{
-		const char* EnvHttpHeaderFile = getenv("UNSYNC_HTTP_HEADER_FILE");
-		if (EnvHttpHeaderFile)
-		{
-			UNSYNC_VERBOSE(L"Using UNSYNC_HTTP_HEADER_FILE environment: '%hs'", EnvHttpHeaderFile);
-			HttpHeaderFilenameUtf8 = std::string(EnvHttpHeaderFile);
-		}
-	}
-
 	if (!PreferredDfsUtf8.empty() && !SourceFilenameUtf8.empty())
 	{
 		LogGlobalStatus(L"Enumerating DFS");
@@ -828,7 +835,8 @@
 		LoginOptions.bInteractive  = bInteractive;
 		LoginOptions.bDecode	   = bDecode;
 		LoginOptions.bPrint		   = bPrint;
-		LoginOptions.bForceRefresh = bForceRefreshAuthentication;
+		LoginOptions.bForceRefresh = bForceRefreshAuth;
+		LoginOptions.bQuick		   = bQuickLogin;
 		return CmdLogin(LoginOptions);
 	}
 	else if (Cli.got_subcommand(SubMount))
