unsync - Trim old cached manifest files during sync

diff --git a/Engine/Source/Programs/Unsync/Private/UnsyncCore.cpp b/Engine/Source/Programs/Unsync/Private/UnsyncCore.cpp
--- a/Engine/Source/Programs/Unsync/Private/UnsyncCore.cpp
+++ b/Engine/Source/Programs/Unsync/Private/UnsyncCore.cpp
@@ -3176,6 +3176,53 @@ struct FFileSyncTaskBatch
 	}
 };
 
+static void
+DeleteOldFilesInDirectory(FPath& Path, uint32 MaxFilesToKeep)
+{
+	struct FEntry
+	{
+		FPath  Path;
+		uint64 Mtime;
+	};
+
+	std::vector<FEntry> Entries;
+
+	FPath ExtendedPath = MakeExtendedAbsolutePath(Path);
+	for (const std::filesystem::directory_entry& It : std::filesystem::directory_iterator(ExtendedPath))
+	{
+		if (It.is_regular_file())
+		{
+			FEntry Entry;
+			Entry.Mtime = ToWindowsFileTime(It.last_write_time());
+			Entry.Path	= It.path();
+			Entries.push_back(Entry);
+		}
+	}
+
+	// reverse sort
+	std::sort(Entries.begin(), Entries.end(), [](const FEntry& A, const FEntry& B) { return A.Mtime > B.Mtime; });
+
+	while (Entries.size() > MaxFilesToKeep)
+	{
+		const FEntry& Oldest = Entries.back();
+
+		std::wstring PathStr = RemoveExtendedPathPrefix(Oldest.Path).wstring();
+
+		if (GDryRun)
+		{
+			UNSYNC_VERBOSE(L"Deleting '%ls'(skipped due to dry run mode)", PathStr.c_str());
+		}
+		else
+		{
+			UNSYNC_VERBOSE(L"Deleting '%ls'", PathStr.c_str());
+			std::error_code ErrorCode = {};
+			FileRemove(Oldest.Path, ErrorCode);
+		}
+
+		Entries.pop_back();
+	}
+}
+
 bool  // TODO: return a TResult
 SyncDirectory(const FSyncDirectoryOptions& SyncOptions)
 {
@@ -3221,6 +3268,14 @@ SyncDirectory(const FSyncDirectoryOptions& SyncOptions)
 		return false;
 	}
 
+	// Delete oldest cached manifest files if there are more than N
+	{
+		UNSYNC_VERBOSE(L"Cleaning temporary directory");
+		UNSYNC_LOG_INDENT;
+		const uint32 MaxFilesToKeep = uint32(5 + SyncOptions.Overlays.size());
+		DeleteOldFilesInDirectory(TargetTempPath, MaxFilesToKeep);
+	}
+
 	FPath		  LogFilePath = TargetManifestRoot / L"unsync.log";
 	FLogFileScope LogFileScope(LogFilePath.wstring().c_str());
 	SetCrashDumpPath(TargetManifestRoot);
diff --git a/Engine/Source/Programs/Unsync/Private/UnsyncMain.cpp b/Engine/Source/Programs/Unsync/Private/UnsyncMain.cpp
--- a/Engine/Source/Programs/Unsync/Private/UnsyncMain.cpp
+++ b/Engine/Source/Programs/Unsync/Private/UnsyncMain.cpp
@@ -330,21 +330,21 @@ InnerMain(int Argc, char** Argv)
 	_setmode(_fileno(stdout), _O_U8TEXT);
 #endif	// UNSYNC_PLATFORM_WINDOWS
 
-	UNSYNC_VERBOSE(L"UNSYNC %hs", GetVersionString().c_str());
+	// Configure default output mehtod based on subcommand.
+	// In machine-readable mode, all verbose logging is directed to stderr.
 
-	if (bUseDebugMode)
+	if (Cli.got_subcommand(SubQuery) || Cli.got_subcommand(SubLogin))
 	{
-		UNSYNC_LOG(L"*** Debug mode enabled ***");
+		GLogMachineReadable = true;
 	}
 
-	UnsyncMallocInit(bUseDebugMode ? EMallocType::Debug : EMallocType::Default);
+	UNSYNC_VERBOSE(L"UNSYNC %hs", GetVersionString().c_str());
 
-	// Configure default output mehtod based on subcommand.
-	// In machine-readable mode, all verbose logging is directed to stderr.
+	UnsyncMallocInit(bUseDebugMode ? EMallocType::Debug : EMallocType::Default);
 
-	if (Cli.got_subcommand(SubQuery) || Cli.got_subcommand(SubLogin))
+	if (bUseDebugMode)
 	{
-		GLogMachineReadable = true;
+		UNSYNC_LOG(L"*** Debug mode enabled ***");
 	}
 
 	// Augment configuration based on environment variables if corresponding command line arguments are missing.
