unsync - Add support for HTTP HEAD requests

Also minor fixes and improvements:

* Fix error code handling when connecting a socket
* Fix some of the naming convention violations
* Add --print-http-header option to login command

diff --git a/Engine/Source/Programs/Unsync/Private/UnsyncCmdHash.cpp b/Engine/Source/Programs/Unsync/Private/UnsyncCmdHash.cpp
--- a/Engine/Source/Programs/Unsync/Private/UnsyncCmdHash.cpp
+++ b/Engine/Source/Programs/Unsync/Private/UnsyncCmdHash.cpp
@@ -59,7 +59,7 @@ CmdHash(const FCmdHashOptions& Options)
 	{
 		UNSYNC_VERBOSE(L"'%ls' is a file", Options.Input.wstring().c_str());
 
-		NativeFile OverlappedFile(Options.Input);
+		FNativeFile OverlappedFile(Options.Input);
 		if (OverlappedFile.IsValid())
 		{
 			UNSYNC_VERBOSE(L"Computing blocks for '%ls' (%.2f MB)", Options.Input.wstring().c_str(), SizeMb(OverlappedFile.GetSize()));
diff --git a/Engine/Source/Programs/Unsync/Private/UnsyncCmdLogin.cpp b/Engine/Source/Programs/Unsync/Private/UnsyncCmdLogin.cpp
--- a/Engine/Source/Programs/Unsync/Private/UnsyncCmdLogin.cpp
+++ b/Engine/Source/Programs/Unsync/Private/UnsyncCmdLogin.cpp
@@ -79,7 +79,16 @@ CmdLogin(const FCmdLoginOptions& Options)
 			}
 			else
 			{
-				LogPrintf(ELogLevel::MachineReadable, L"%hs\n", AuthTokenResult->Raw.c_str());
+				if (Options.bPrintHttpHeader)
+				{
+					std::string HttpHeader = "Authorization: Bearer " + AuthTokenResult->Access;
+					LogPrintf(ELogLevel::MachineReadable, L"%hs\n", HttpHeader.c_str());
+				}
+				else
+				{
+					LogPrintf(ELogLevel::MachineReadable, L"%hs\n", AuthTokenResult->Raw.c_str());
+				}
+				
 			}
 		}
 
diff --git a/Engine/Source/Programs/Unsync/Private/UnsyncCmdLogin.h b/Engine/Source/Programs/Unsync/Private/UnsyncCmdLogin.h
--- a/Engine/Source/Programs/Unsync/Private/UnsyncCmdLogin.h
+++ b/Engine/Source/Programs/Unsync/Private/UnsyncCmdLogin.h
@@ -10,11 +10,12 @@ namespace unsync {
 struct FCmdLoginOptions
 {
 	FRemoteDesc Remote;
-	bool		bInteractive  = false;	// TODO
-	bool		bPrint		  = false;
-	bool		bDecode		  = false;
-	bool		bForceRefresh = false;
-	bool		bQuick		  = false;
+	bool		bInteractive	 = false;  // TODO
+	bool		bPrint			 = false;
+	bool		bPrintHttpHeader = false;
+	bool		bDecode			 = false;
+	bool		bForceRefresh	 = false;
+	bool		bQuick			 = false;
 };
 
 int32 CmdLogin(const FCmdLoginOptions& Options);
diff --git a/Engine/Source/Programs/Unsync/Private/UnsyncCore.cpp b/Engine/Source/Programs/Unsync/Private/UnsyncCore.cpp
--- a/Engine/Source/Programs/Unsync/Private/UnsyncCore.cpp
+++ b/Engine/Source/Programs/Unsync/Private/UnsyncCore.cpp
@@ -1816,7 +1816,7 @@ CreateDirectoryManifest(const FPath& Root, uint32 BlockSize, FAlgorithmOptions A
 		if (BlockSize)
 		{
 			FPath FilePath = Root / RelativePath;
-			auto  File	   = std::make_shared<NativeFile>(FilePath, EFileMode::ReadOnlyUnbuffered);
+			auto  File	   = std::make_shared<FNativeFile>(FilePath, EFileMode::ReadOnlyUnbuffered);
 			if (File->IsValid())
 			{
 				UNSYNC_VERBOSE(L"Computing blocks for '%ls' (%.2f MB)", FilePath.wstring().c_str(), double(File->GetSize()) / (1 << 20));
@@ -1962,7 +1962,7 @@ UpdateDirectoryManifestBlocks(FDirectoryManifest& Result, const FPath& Root, uin
 		++NumProcessedFiles;
 
 		FPath FilePath = Root / It.first;
-		auto  File	   = std::make_shared<NativeFile>(FilePath, EFileMode::ReadOnlyUnbuffered);
+		auto  File	   = std::make_shared<FNativeFile>(FilePath, EFileMode::ReadOnlyUnbuffered);
 		if (File->IsValid())
 		{
 			UNSYNC_VERBOSE(L"Computing blocks for '%ls' (%.2f MB)", FilePath.wstring().c_str(), double(File->GetSize()) / (1 << 20));
@@ -2364,7 +2364,7 @@ SyncFile(const FNeedList&		   NeedList,
 				CreateDirectories(TargetFileParent);
 			}
 
-			auto TargetFile = NativeFile(TargetFilePath, EFileMode::CreateWriteOnly, 0);
+			auto TargetFile = FNativeFile(TargetFilePath, EFileMode::CreateWriteOnly, 0);
 			if (TargetFile.IsValid())
 			{
 				Result.Status = EFileSyncStatus::Ok;
@@ -2406,7 +2406,7 @@ SyncFile(const FNeedList&		   NeedList,
 				CreateDirectories(TargetFileParent);
 			}
 
-			TargetFile = std::make_unique<NativeFile>(TempTargetFilePath, EFileMode::CreateWriteOnly, TargetFileSizeInfo.TotalBytes);
+			TargetFile = std::make_unique<FNativeFile>(TempTargetFilePath, EFileMode::CreateWriteOnly, TargetFileSizeInfo.TotalBytes);
 			if (TargetFile->GetError() != 0)
 			{
 				UNSYNC_ERROR(L"Failed to create output file '%ls'. Error code %d.",
@@ -2420,7 +2420,7 @@ SyncFile(const FNeedList&		   NeedList,
 		FDeferredOpenReader SourceFile([SourceFilePath, TargetFilePath] {
 			UNSYNC_VERBOSE(L"Opening source file '%ls'", SourceFilePath.wstring().c_str());
 			LogStatus(TargetFilePath.wstring().c_str(), L"Opening source file");
-			return std::unique_ptr<NativeFile>(new NativeFile(SourceFilePath, EFileMode::ReadOnlyUnbuffered));
+			return std::unique_ptr<FNativeFile>(new FNativeFile(SourceFilePath, EFileMode::ReadOnlyUnbuffered));
 		});
 
 		FBuildTargetResult BuildResult = BuildTarget(*TargetFile,
@@ -2445,7 +2445,7 @@ SyncFile(const FNeedList&		   NeedList,
 			{
 				// Reopen the file in unuffered read mode for optimal reading performance
 				TargetFile = nullptr;
-				TargetFile = std::make_unique<NativeFile>(TempTargetFilePath, EFileMode::ReadOnlyUnbuffered);
+				TargetFile = std::make_unique<FNativeFile>(TempTargetFilePath, EFileMode::ReadOnlyUnbuffered);
 			}
 
 			if (!ValidateTarget(*TargetFile, NeedList, Options.Algorithm.StrongHashAlgorithmId))
@@ -2546,7 +2546,7 @@ SyncFile(const FPath& SourceFilePath, const FPath& BaseFilePath, const FPath& Ta
 
 	FFileSyncResult Result;
 
-	NativeFile BaseFile(BaseFilePath, EFileMode::ReadOnlyUnbuffered);
+	FNativeFile BaseFile(BaseFilePath, EFileMode::ReadOnlyUnbuffered);
 	if (!BaseFile.IsValid())
 	{
 		BaseFile.Close();
@@ -3036,8 +3036,15 @@ LoadAndMergeSourceManifest(FDirectoryManifest& Output,
 
 	if (bDownloadManifestFromProxy)
 	{
+		// TODO: add a mechanism to suppress dry run in a scope in a thread-safe way
+		const bool bPrevDryRun = GDryRun;
+		GDryRun				   = false;
+
 		UNSYNC_LOG_INDENT;
 		bool bDownloadedOk = DownloadFileIfNewer(ProxyPool.RemoteDesc, SourceManifestPath, SourceManifestTempPath);
+
+		GDryRun = bPrevDryRun;
+
 		if (!bDownloadedOk)
 		{
 			UNSYNC_ERROR(L"Failed to download manifest file '%ls'", SourceManifestPath.wstring().c_str());
@@ -3588,7 +3595,7 @@ SyncDirectory(const FSyncDirectoryOptions& SyncOptions)
 
 			if (Item.IsBaseValid() && PathExists(Item.BaseFilePath))
 			{
-				NativeFile BaseFile(Item.BaseFilePath, EFileMode::ReadOnlyUnbuffered);
+				FNativeFile BaseFile(Item.BaseFilePath, EFileMode::ReadOnlyUnbuffered);
 				uint32	   SourceBlockSize = Item.SourceManifest->BlockSize;
 				UNSYNC_VERBOSE(L"Computing difference for target '%ls' (base size: %.2f MB)",
 							   Item.BaseFilePath.wstring().c_str(),
@@ -3816,10 +3823,10 @@ SyncDirectory(const FSyncDirectoryOptions& SyncOptions)
 								 &ProxyPool](const FFileSyncTask& Item, FBlockCache* BlockCache, bool bBackground) {
 			UNSYNC_VERBOSE(L"Copy '%ls' (%ls)", Item.TargetFilePath.wstring().c_str(), (Item.NeedBytesFromBase) ? L"partial" : L"full");
 
-			std::unique_ptr<NativeFile> BaseFile;
+			std::unique_ptr<FNativeFile> BaseFile;
 			if (Item.IsBaseValid())
 			{
-				BaseFile = std::make_unique<NativeFile>(Item.BaseFilePath, EFileMode::ReadOnlyUnbuffered);
+				BaseFile = std::make_unique<FNativeFile>(Item.BaseFilePath, EFileMode::ReadOnlyUnbuffered);
 			}
 
 			const FGenericBlockArray& SourceBlocks	  = Item.SourceManifest->Blocks;
diff --git a/Engine/Source/Programs/Unsync/Private/UnsyncFile.cpp b/Engine/Source/Programs/Unsync/Private/UnsyncFile.cpp
--- a/Engine/Source/Programs/Unsync/Private/UnsyncFile.cpp
+++ b/Engine/Source/Programs/Unsync/Private/UnsyncFile.cpp
@@ -263,7 +263,7 @@ FWindowsFile::CompleteReadCommand(Command& Cmd)
 		}
 		else
 		{
-			UNSYNC_WARNING(L"NativeFile expected to read %lld bytes, but read %lld. Last error code: %d.",
+			UNSYNC_WARNING(L"FNativeFile expected to read %lld bytes, but read %lld. Last error code: %d.",
 						   (uint64)ExpectedReadBytes,
 						   (uint64)ReadBytes,
 						   LastError);
@@ -854,7 +854,7 @@ FBuffer
 ReadFileToBuffer(const FPath& Filename)
 {
 	FBuffer	   Result;
-	NativeFile File(Filename, EFileMode::ReadOnly);
+	FNativeFile File(Filename, EFileMode::ReadOnly);
 	if (File.IsValid())
 	{
 		Result.Resize(File.GetSize());
@@ -872,7 +872,7 @@ WriteBufferToFile(const FPath& Filename, const uint8* Data, uint64 Size)
 	UNSYNC_ASSERT(Size);
 	UNSYNC_ASSERT(!GDryRun);
 
-	NativeFile File(Filename, EFileMode::CreateReadWrite, Size);
+	FNativeFile File(Filename, EFileMode::CreateReadWrite, Size);
 
 	if (File.IsValid())
 	{
diff --git a/Engine/Source/Programs/Unsync/Private/UnsyncFile.h b/Engine/Source/Programs/Unsync/Private/UnsyncFile.h
--- a/Engine/Source/Programs/Unsync/Private/UnsyncFile.h
+++ b/Engine/Source/Programs/Unsync/Private/UnsyncFile.h
@@ -205,7 +205,7 @@ private:
 private:
 	EFileMode Mode;
 };
-using NativeFile = FWindowsFile;
+using FNativeFile = FWindowsFile;
 #endif	// UNSYNC_PLATFORM_WINDOWS
 
 #if UNSYNC_PLATFORM_UNIX
@@ -244,7 +244,7 @@ private:
 	FILE*	  FileHandle	 = nullptr;
 	int		  FileDescriptor = 0;
 };
-using NativeFile = FUnixFile;
+using FNativeFile = FUnixFile;
 #endif	// UNSYNC_PLATFORM_UNIX
 
 struct FVectorStreamOut
diff --git a/Engine/Source/Programs/Unsync/Private/UnsyncHttp.cpp b/Engine/Source/Programs/Unsync/Private/UnsyncHttp.cpp
--- a/Engine/Source/Programs/Unsync/Private/UnsyncHttp.cpp
+++ b/Engine/Source/Programs/Unsync/Private/UnsyncHttp.cpp
@@ -43,8 +43,13 @@ struct FHttpParser
 {
 	static FHttpParser* ToThis(http_parser* Parser) { return (FHttpParser*)(Parser->data); }
 
-	FHttpParser(HttpMessageCallback InResponseCallback, uint8* InScratchBuffer, uint64 InScratchSize, http_parser_type Type)
+	FHttpParser(HttpMessageCallback InResponseCallback,
+				uint8*				InScratchBuffer,
+				uint64				InScratchSize,
+				http_parser_type	Type,
+				EHttpMethod			InMethod)
 	: ResponseCallback(InResponseCallback)
+	, Method(InMethod)
 	, ScratchBuffer(InScratchBuffer)
 	, ScratchSize(InScratchSize)
 	{
@@ -112,7 +117,8 @@ struct FHttpParser
 
 		if (MatchUncased(PendingHeader, "content-length"))
 		{
-			ContentLength = atoi(Data);
+			ContentLength		   = strtoull(Data, nullptr, 10);
+			Response.ContentLength = uint64(ContentLength);
 			Response.Buffer.Reserve(ContentLength);
 		}
 		else if (MatchUncased(PendingHeader, "content-type"))
@@ -150,7 +156,14 @@ struct FHttpParser
 		Response.Code	= Parser.status_code;
 		bHeaderComplete = true;
 
-		return 0;
+		if (Method == EHttpMethod::HEAD)
+		{
+			return 1;
+		}
+		else
+		{
+			return 0;
+		}
 	}
 
 	int OnStatus(const char* Data, size_t Size) { return 0; }
@@ -203,6 +216,8 @@ struct FHttpParser
 	HttpMessageCallback ResponseCallback;
 	FHttpResponse		Response;
 
+	EHttpMethod Method = EHttpMethod::GET;
+
 	http_parser_settings Settings;
 	http_parser			 Parser;
 
@@ -216,7 +231,7 @@ struct FHttpParser
 	std::string_view PendingHeader = {};
 	std::string_view PendingValue  = {};
 
-	int32 ContentLength = 0;
+	uint64 ContentLength = 0;
 
 	uint64 TotalReceivedBytes = 0;
 	uint64 TotalParsedBytes	  = 0;
@@ -275,6 +290,14 @@ HttpRequestBegin(FHttpConnection& Connection, const FHttpRequest& Request)
 		return false;
 	}
 
+	if (Connection.NumActiveRequests > 0 && Connection.Method != Request.Method)
+	{
+		UNSYNC_ERROR(L"HTTP connection must not have outstanding requests when request method is changed");
+		return false;
+	}
+
+	Connection.Method = Request.Method;
+
 	// TODO: use a string builder
 	std::string HttpHeader;
 	switch (Request.Method)
@@ -285,9 +308,9 @@ HttpRequestBegin(FHttpConnection& Connection, const FHttpRequest& Request)
 		case EHttpMethod::GET:
 			HttpHeader = "GET ";
 			break;
-			// 	case HttpMethod::HEAD: // < TODO: support requests that don't return a body
-			// 		http_header = "HEAD";
-			// 		break;
+		case EHttpMethod::HEAD:
+			HttpHeader = "HEAD ";
+			break;
 		case EHttpMethod::POST:
 			HttpHeader = "POST ";
 			break;
@@ -439,7 +462,7 @@ HttpRequestEnd(FHttpConnection& Connection)
 	// TODO: user-provided scratch buffer
 	uint8 ScratchBuffer[256_KB];
 	ScratchBuffer[0] = 0;
-	FHttpParser Parser(MessageCallback, ScratchBuffer, sizeof(ScratchBuffer), HTTP_RESPONSE);
+	FHttpParser Parser(MessageCallback, ScratchBuffer, sizeof(ScratchBuffer), HTTP_RESPONSE, Connection.Method);
 
 	while (!Parser.bComplete)
 	{
@@ -526,7 +549,7 @@ FHttpConnection::Open()
 
 	FSocketHandle RawSocketHandle = SocketConnectTcp(HostAddress.c_str(), HostPort);
 
-	if (RawSocketHandle < 0)
+	if (RawSocketHandle == InvalidSocketHandle)
 	{
 		return false;
 	}
diff --git a/Engine/Source/Programs/Unsync/Private/UnsyncHttp.h b/Engine/Source/Programs/Unsync/Private/UnsyncHttp.h
--- a/Engine/Source/Programs/Unsync/Private/UnsyncHttp.h
+++ b/Engine/Source/Programs/Unsync/Private/UnsyncHttp.h
@@ -29,7 +29,7 @@ enum class EHttpContentType
 enum class EHttpMethod
 {
 	GET,
-	// HEAD, // < TODO: support requests that don't return a body
+	HEAD,
 	POST,
 	PUT,
 };
@@ -50,6 +50,7 @@ struct FHttpRequest
 struct FHttpResponse
 {
 	FBuffer Buffer;	 // TODO: use pooled IOBuffer
+	uint64	ContentLength = 0;
 	int32	Code = 0;
 
 	EHttpContentType ContentType = EHttpContentType::Unknown;
@@ -86,6 +87,8 @@ struct FHttpConnection
 	// TODO: use single memory allocation for multiple reponse objects (perhaps a ring buffer)
 	std::deque<FHttpResponse> ResponseQueue;  // Contains HTTP responses for pipelined requests
 
+	EHttpMethod Method = EHttpMethod::GET;
+
 	FTimePoint LastUsed = {};
 
 	FSocketBase& GetSocket()
diff --git a/Engine/Source/Programs/Unsync/Private/UnsyncJupiter.cpp b/Engine/Source/Programs/Unsync/Private/UnsyncJupiter.cpp
--- a/Engine/Source/Programs/Unsync/Private/UnsyncJupiter.cpp
+++ b/Engine/Source/Programs/Unsync/Private/UnsyncJupiter.cpp
@@ -538,7 +538,7 @@ JupiterPush(const FDirectoryManifest& Manifest, const FRemoteDesc& RemoteDesc, F
 
 		const FFileManifest& FileManifest = It.second;
 
-		NativeFile File(FileManifest.CurrentPath, EFileMode::ReadOnlyUnbuffered);
+		FNativeFile File(FileManifest.CurrentPath, EFileMode::ReadOnlyUnbuffered);
 		if (!File.IsValid())
 		{
 			UNSYNC_ERROR(L"Failed to open input file '%ls'", FileManifest.CurrentPath.wstring().c_str());
diff --git a/Engine/Source/Programs/Unsync/Private/UnsyncMain.cpp b/Engine/Source/Programs/Unsync/Private/UnsyncMain.cpp
--- a/Engine/Source/Programs/Unsync/Private/UnsyncMain.cpp
+++ b/Engine/Source/Programs/Unsync/Private/UnsyncMain.cpp
@@ -89,6 +89,7 @@ InnerMain(int Argc, char** Argv)
 	bool					 bInteractive		 = false;
 	bool					 bDecode			 = false;
 	bool					 bPrint				 = false;
+	bool					 bPrintHttpHeader	 = false;
 	bool					 bShouldLogin		 = false;
 	bool					 bQuickLogin		 = false;
 	bool					 bForceRefreshAuth	 = false;
@@ -285,6 +286,7 @@ InnerMain(int Argc, char** Argv)
 	SubLogin->add_flag("--interactive", bInteractive, "Allow user interaction through modal dialogs");
 	SubLogin->add_flag("--decode", bDecode, "Decode authentication token (implies --print)");
 	SubLogin->add_flag("--print", bPrint, "Print authentication token to standard output");
+	SubLogin->add_flag("--print-http-header", bPrintHttpHeader, "Print authentication token to standard output as HTTP Authorization header that could be used with curl, etc.");
 	SubLogin->add_flag("--refresh", bForceRefreshAuth, "Force authentication refresh even if access token has not yet expired");
 	SubLogin->add_flag("--quick", bQuickLogin, "Skip token validation using remote server (fast path when cached acess token is expected to be valid)");
 	AddTlsOptions(SubLogin);
@@ -825,7 +827,7 @@ InnerMain(int Argc, char** Argv)
 	}
 	else if (Cli.got_subcommand(SubLogin))
 	{
-		if (bDecode)
+		if (bDecode || bPrintHttpHeader)
 		{
 			bPrint = true;
 		}
@@ -835,6 +837,7 @@ InnerMain(int Argc, char** Argv)
 		LoginOptions.bInteractive  = bInteractive;
 		LoginOptions.bDecode	   = bDecode;
 		LoginOptions.bPrint		   = bPrint;
+		LoginOptions.bPrintHttpHeader = bPrintHttpHeader;
 		LoginOptions.bForceRefresh = bForceRefreshAuth;
 		LoginOptions.bQuick		   = bQuickLogin;
 		return CmdLogin(LoginOptions);
diff --git a/Engine/Source/Programs/Unsync/Private/UnsyncScavenger.cpp b/Engine/Source/Programs/Unsync/Private/UnsyncScavenger.cpp
--- a/Engine/Source/Programs/Unsync/Private/UnsyncScavenger.cpp
+++ b/Engine/Source/Programs/Unsync/Private/UnsyncScavenger.cpp
@@ -366,7 +366,7 @@ BuildTargetFromScavengedData(FIOWriter&						Output,
 
 		const FFileManifest& ScavengeFileManifest = ScavengeDatabase.GetFileManifest(PossibleSource);
 
-		NativeFile LocalSourceFile = NativeFile(PossibleSource.FullSourceFilePath, EFileMode::ReadOnlyUnbuffered);
+		FNativeFile LocalSourceFile = FNativeFile(PossibleSource.FullSourceFilePath, EFileMode::ReadOnlyUnbuffered);
 		if (LocalSourceFile.IsValid())
 		{
 			THashMap<FGenericHash, uint64> BlockOffsetMap;
diff --git a/Engine/Source/Programs/Unsync/Private/UnsyncSerialization.cpp b/Engine/Source/Programs/Unsync/Private/UnsyncSerialization.cpp
--- a/Engine/Source/Programs/Unsync/Private/UnsyncSerialization.cpp
+++ b/Engine/Source/Programs/Unsync/Private/UnsyncSerialization.cpp
@@ -120,7 +120,7 @@ SaveBlocks(const std::vector<FBlock128>& Blocks, uint32 BlockSize, const FPath&
 {
 	const uint64 OutputSize = sizeof(FBlockFileHeader) + sizeof(FBlock128) * Blocks.size();
 
-	NativeFile File(Filename, EFileMode::CreateReadWrite, OutputSize);
+	FNativeFile File(Filename, EFileMode::CreateReadWrite, OutputSize);
 
 	if (File.IsValid())
 	{
@@ -516,7 +516,7 @@ LoadDirectoryManifest(FDirectoryManifest& OutManifest, const FPath& Root, const
 	UNSYNC_LOG_INDENT;
 	UNSYNC_VERBOSE(L"Loading directory manifest from '%ls'", ManifestFilename.wstring().c_str());
 
-	// NativeFile manifest_file(manifest_filename, FileMode::ReadOnly);
+	// FNativeFile manifest_file(manifest_filename, FileMode::ReadOnly);
 	FBuffer ManifestBuffer = ReadFileToBuffer(ManifestFilename);
 
 	if (ManifestBuffer.Size() != 0)
@@ -646,7 +646,7 @@ SaveDirectoryManifest(const FDirectoryManifest& Manifest, const FPath& Filename)
 	bool bSerialized = SaveDirectoryManifest(Manifest, OutputStream);
 	UNSYNC_ASSERT(bSerialized);
 
-	NativeFile OutputFile(Filename, EFileMode::CreateWriteOnly, OutputBuffer.Size());
+	FNativeFile OutputFile(Filename, EFileMode::CreateWriteOnly, OutputBuffer.Size());
 	if (OutputFile.IsValid())
 	{
 		uint64 WroteBytes = OutputFile.Write(OutputBuffer.Data(), 0, OutputBuffer.Size());
diff --git a/Engine/Source/Programs/Unsync/Private/UnsyncSocket.cpp b/Engine/Source/Programs/Unsync/Private/UnsyncSocket.cpp
--- a/Engine/Source/Programs/Unsync/Private/UnsyncSocket.cpp
+++ b/Engine/Source/Programs/Unsync/Private/UnsyncSocket.cpp
@@ -41,6 +41,8 @@ namespace unsync {
 static_assert(sizeof(FSocketAddress) >= sizeof(sockaddr), "SocketAddress is too small");
 static_assert(sizeof(FSocketAddress) >= sizeof(sockaddr_in), "SocketAddress is too small");
 
+const FSocketHandle InvalidSocketHandle = INVALID_SOCKET;
+
 struct FSocketInitHelper
 {
 	FSocketInitHelper()
@@ -103,10 +105,10 @@ SocketListenTcp(const char* Address, uint16 Port)
 
 	SOCKET ListenSocket = socket(AF_INET, SOCK_STREAM, 0);
 
-	if (ListenSocket == INVALID_SOCKET)
+	if (ListenSocket == InvalidSocketHandle)
 	{
 		UNSYNC_ERROR(L"Failed to create TCP socket (error code %d)", GetLastSocketError());
-		return 0;
+		return InvalidSocketHandle;
 	}
 
 	sockaddr_in Service;
@@ -119,7 +121,7 @@ SocketListenTcp(const char* Address, uint16 Port)
 	{
 		UNSYNC_ERROR(L"Failed to bind TCP socket (error code %d)", GetLastSocketError());
 		SocketClose(ListenSocket);
-		return 0;
+		return InvalidSocketHandle;
 	}
 
 	int32 ListenResult = listen(ListenSocket, 1);
@@ -127,7 +129,7 @@ SocketListenTcp(const char* Address, uint16 Port)
 	{
 		UNSYNC_ERROR(L"Failed to listen on TCP socket (error code %d)", GetLastSocketError());
 		SocketClose(ListenSocket);
-		return 0;
+		return InvalidSocketHandle;
 	}
 
 	return ListenSocket;
@@ -137,10 +139,10 @@ FSocketHandle
 SocketAccept(FSocketHandle ListenSocket)
 {
 	SOCKET AcceptSocket = accept(ListenSocket, nullptr, nullptr);
-	if (AcceptSocket == INVALID_SOCKET)
+	if (AcceptSocket == InvalidSocketHandle)
 	{
 		UNSYNC_ERROR(L"Failed to accept connection on TCP socket (error code %d)", GetLastSocketError());
-		return 0;
+		return InvalidSocketHandle;
 	}
 	return AcceptSocket;
 }
@@ -152,10 +154,10 @@ SocketConnectTcp(const char* DestAddress, uint16 Port)
 
 	SOCKET Sock = socket(AF_INET, SOCK_STREAM, 0);
 
-	if (Sock == INVALID_SOCKET)
+	if (Sock == InvalidSocketHandle)
 	{
 		UNSYNC_ERROR(L"Failed to create TCP socket (error code %d)", GetLastSocketError());
-		return 0;
+		return InvalidSocketHandle;
 	}
 
 	sockaddr_in SockAddr = {};
@@ -193,7 +195,7 @@ SocketConnectTcp(const char* DestAddress, uint16 Port)
 	else
 	{
 		UNSYNC_ERROR(L"Invalid address '%hs'", DestAddress);
-		return 0;
+		return InvalidSocketHandle;
 	}
 
 	if (Addr)
@@ -204,13 +206,15 @@ SocketConnectTcp(const char* DestAddress, uint16 Port)
 	if (inet_pton(SockAddr.sin_family, DestAddress, &SockAddr.sin_addr) <= 0)
 	{
 		UNSYNC_ERROR(L"Invalid address '%hs'", DestAddress);
-		return 0;
+		return InvalidSocketHandle;
 	}
 
-	if (connect(Sock, (struct sockaddr*)&SockAddr, sizeof(SockAddr)) < 0)
+	int32 ConnectResult = connect(Sock, (struct sockaddr*)&SockAddr, sizeof(SockAddr));
+	if (ConnectResult < 0)
 	{
+		closesocket(Sock);
 		UNSYNC_LOG(L"Warning: Failed to connect to '%hs' on port %d", DestAddress, Port);
-		return 0;
+		return InvalidSocketHandle;
 	}
 
 	return Sock;
@@ -219,7 +223,7 @@ SocketConnectTcp(const char* DestAddress, uint16 Port)
 void
 SocketClose(FSocketHandle Socket)
 {
-	if (Socket)
+	if (Socket != InvalidSocketHandle)
 	{
 		closesocket(Socket);
 	}
diff --git a/Engine/Source/Programs/Unsync/Private/UnsyncSocket.h b/Engine/Source/Programs/Unsync/Private/UnsyncSocket.h
--- a/Engine/Source/Programs/Unsync/Private/UnsyncSocket.h
+++ b/Engine/Source/Programs/Unsync/Private/UnsyncSocket.h
@@ -14,6 +14,8 @@ class FBuffer;
 struct FBufferView;
 using FSocketHandle = uintptr_t;
 
+extern const FSocketHandle InvalidSocketHandle;
+
 struct FSocketAddress  // opaque buffer that holds socket address
 {
 	FSocketAddress()
diff --git a/Engine/Source/Programs/Unsync/Private/UnsyncTest.cpp b/Engine/Source/Programs/Unsync/Private/UnsyncTest.cpp
--- a/Engine/Source/Programs/Unsync/Private/UnsyncTest.cpp
+++ b/Engine/Source/Programs/Unsync/Private/UnsyncTest.cpp
@@ -472,7 +472,7 @@ TestFiles()
 
 	{
 		// overlapped file writing
-		NativeFile TestFile(TestFilename, EFileMode::CreateReadWrite, TestFileSize);
+		FNativeFile TestFile(TestFilename, EFileMode::CreateReadWrite, TestFileSize);
 		UNSYNC_ASSERT(TestFile.IsValid());
 
 		FillRandomBytes(TempBuffer.get(), TestFileSize, 321);
@@ -489,7 +489,7 @@ TestFiles()
 	{
 		memset(TempBuffer.get(), 0, TestFileSize);
 
-		NativeFile TestFile(TestFilename, EFileMode::ReadOnly);
+		FNativeFile TestFile(TestFilename, EFileMode::ReadOnly);
 		UNSYNC_ASSERT(TestFile.IsValid());
 		UNSYNC_ASSERT(TestFile.GetSize() == TestFileSize);
 		TestFile.Read(TempBuffer.get(), 0, TestFileSize);
