unsync - Use multiple concurrent connections during file download query

* Bump version to 1.0.54

diff --git a/Engine/Source/Programs/Unsync/Private/UnsyncCmdQuery.cpp b/Engine/Source/Programs/Unsync/Private/UnsyncCmdQuery.cpp
--- a/Engine/Source/Programs/Unsync/Private/UnsyncCmdQuery.cpp
+++ b/Engine/Source/Programs/Unsync/Private/UnsyncCmdQuery.cpp
@@ -110,7 +110,7 @@ CmdQueryMirrors(const FCmdQueryOptions& Options)
 
 	std::vector<FMirrorInfo> Mirrors = MirrorsResult.GetData();
 
-	ParallelForEach(Mirrors.begin(), Mirrors.end(), [](FMirrorInfo& Mirror) { Mirror.Ping = RunHttpPing(Mirror.Address, Mirror.Port); });
+	ParallelForEach(Mirrors, [](FMirrorInfo& Mirror) { Mirror.Ping = RunHttpPing(Mirror.Address, Mirror.Port); });
 
 	std::sort(Mirrors.begin(), Mirrors.end(), [](const FMirrorInfo& InA, const FMirrorInfo& InB) {
 		double A = InA.Ping > 0 ? InA.Ping : FLT_MAX;
@@ -178,28 +178,22 @@ CmdQueryList(const FCmdQueryOptions& Options)
 int32
 CmdQueryFile(const FCmdQueryOptions& Options)
 {
-	TResult<FBuffer> Response = ProxyQuery::DownloadFile(Options.Remote, Options.Args);
-
 	UNSYNC_LOG(L"Downloading file: '%hs'", Options.Args.c_str());
+	UNSYNC_LOG_INDENT;
 
-	if (Response.IsOk())
-	{
-		FBuffer& Buffer = Response.GetData();
+	std::unique_ptr<FNativeFile> ResultWriter;
 
-		UNSYNC_LOG(L"Downloaded file size: %llu bytes", llu(Buffer.Size()));
+	auto OutputCallback = [&ResultWriter, &Options](uint64 Size) -> FIOWriter& {
+		UNSYNC_LOG(L"Size: %llu bytes (%.3f MB)", llu(Size), SizeMb(Size));
+		ResultWriter = std::make_unique<FNativeFile>(Options.OutputPath, EFileMode::CreateWriteOnly, Size);
+		return *ResultWriter;
+	};
 
-		if (!Options.OutputPath.empty())
-		{
-			bool bWriteOk = WriteBufferToFile(Options.OutputPath, Buffer);
-			if (bWriteOk)
-			{
-				UNSYNC_LOG(L"Output written to file '%s'", Options.OutputPath.wstring().c_str());
-			}
-			else
-			{
-				UNSYNC_ERROR(L"Failed to write output file '%s'", Options.OutputPath.wstring().c_str());
-			}
-		}
+	TResult<> Response = ProxyQuery::DownloadFile(Options.Remote, Options.Args, OutputCallback);
+
+	if (Response.IsOk())
+	{
+		UNSYNC_LOG(L"Output written to file '%s'", Options.OutputPath.wstring().c_str());
 	}
 	else
 	{
@@ -241,7 +235,7 @@ FindClosestMirror(const FRemoteDesc& Remote)
 	}
 
 	std::vector<FMirrorInfo> Mirrors = MirrorsResult.GetData();
-	ParallelForEach(Mirrors.begin(), Mirrors.end(), [](FMirrorInfo& Mirror) { Mirror.Ping = RunHttpPing(Mirror.Address, Mirror.Port); });
+	ParallelForEach(Mirrors, [](FMirrorInfo& Mirror) { Mirror.Ping = RunHttpPing(Mirror.Address, Mirror.Port); });
 
 	std::sort(Mirrors.begin(), Mirrors.end(), [](const FMirrorInfo& InA, const FMirrorInfo& InB) {
 		double A = InA.Ping > 0 ? InA.Ping : FLT_MAX;
diff --git a/Engine/Source/Programs/Unsync/Private/UnsyncCore.cpp b/Engine/Source/Programs/Unsync/Private/UnsyncCore.cpp
--- a/Engine/Source/Programs/Unsync/Private/UnsyncCore.cpp
+++ b/Engine/Source/Programs/Unsync/Private/UnsyncCore.cpp
@@ -27,7 +27,7 @@ UNSYNC_THIRD_PARTY_INCLUDES_START
 #include <md5-sse2.h>
 UNSYNC_THIRD_PARTY_INCLUDES_END
 
-#define UNSYNC_VERSION_STR "1.0.53-dev"
+#define UNSYNC_VERSION_STR "1.0.54"
 
 namespace unsync {
 
@@ -3646,7 +3646,7 @@ SyncDirectory(const FSyncDirectoryOptions& SyncOptions)
 			UNSYNC_ASSERT(Item.NeedBytesFromSource + Item.NeedBytesFromBase == Item.TotalSizeBytes);
 		};
 
-		ParallelForEach(AllFileTasks.begin(), AllFileTasks.end(), DiffTask);
+		ParallelForEach(AllFileTasks, DiffTask);
 
 		auto TimeDiffEnd = TimePointNow();
 
diff --git a/Engine/Source/Programs/Unsync/Private/UnsyncError.h b/Engine/Source/Programs/Unsync/Private/UnsyncError.h
--- a/Engine/Source/Programs/Unsync/Private/UnsyncError.h
+++ b/Engine/Source/Programs/Unsync/Private/UnsyncError.h
@@ -8,6 +8,8 @@
 #include <string>
 #include <type_traits>
 #include <variant>
+#include <optional>
+#include <atomic>
 
 namespace unsync {
 
@@ -149,4 +151,24 @@ ResultOk()
 	return TResult<FEmpty, E>(FEmpty());
 }
 
+struct FAtomicError
+{
+	std::atomic_flag	  Flag;
+	std::optional<FError> Data;
+
+	bool Test() const { return Flag.test(); }
+	bool Set(FError&& InData)
+	{
+		if (Flag.test_and_set() == false)
+		{
+			Data = std::move(InData);
+			return true;
+		}
+		else
+		{
+			return false;
+		}
+	}
+};
+
 }  // namespace unsync
diff --git a/Engine/Source/Programs/Unsync/Private/UnsyncMain.cpp b/Engine/Source/Programs/Unsync/Private/UnsyncMain.cpp
--- a/Engine/Source/Programs/Unsync/Private/UnsyncMain.cpp
+++ b/Engine/Source/Programs/Unsync/Private/UnsyncMain.cpp
@@ -649,8 +649,7 @@ InnerMain(int Argc, char** Argv)
 
 	if (bUseTls)
 	{
-		RemoteDesc.bTlsEnable			 = true;
-		RemoteDesc.bTlsVerifyCertificate = true;
+		RemoteDesc.bTlsEnable = true;
 	}
 
 	if (bAllowInsecureTls)
@@ -659,6 +658,10 @@ InnerMain(int Argc, char** Argv)
 		RemoteDesc.TlsSubject			 = {};
 		UNSYNC_WARNING(L"Remote server certificate verification is disabled.");
 	}
+	else
+	{
+		RemoteDesc.bTlsVerifyCertificate = true;
+	}
 
 	if (bNoOutputValidation)
 	{
diff --git a/Engine/Source/Programs/Unsync/Private/UnsyncProgress.cpp b/Engine/Source/Programs/Unsync/Private/UnsyncProgress.cpp
--- a/Engine/Source/Programs/Unsync/Private/UnsyncProgress.cpp
+++ b/Engine/Source/Programs/Unsync/Private/UnsyncProgress.cpp
@@ -23,7 +23,9 @@ AddGlobalProgress(uint64 Size, EBlockListType ListType)
 }
 
 FLogProgressScope::FLogProgressScope(uint64 InTotal, ELogProgressUnits InUnits, uint64 InPeriodMilliseconds)
-: Current(0)
+: bParentThreadVerbose(GLogVerbose)
+, ParentThreadIndent(GLogIndent)
+, Current(0)
 , Total(InTotal)
 , PeriodMilliseconds(InPeriodMilliseconds)
 , Units(InUnits)
diff --git a/Engine/Source/Programs/Unsync/Private/UnsyncProgress.h b/Engine/Source/Programs/Unsync/Private/UnsyncProgress.h
--- a/Engine/Source/Programs/Unsync/Private/UnsyncProgress.h
+++ b/Engine/Source/Programs/Unsync/Private/UnsyncProgress.h
@@ -48,6 +48,9 @@ struct FLogProgressScope
 
 	void Add(uint64 X, bool bForceComplete = false);
 
+	const bool	 bParentThreadVerbose;
+	const uint32 ParentThreadIndent;
+
 	std::mutex				Mutex;
 	std::atomic<uint64>		Current;
 	const uint64			Total;
diff --git a/Engine/Source/Programs/Unsync/Private/UnsyncProxy.cpp b/Engine/Source/Programs/Unsync/Private/UnsyncProxy.cpp
--- a/Engine/Source/Programs/Unsync/Private/UnsyncProxy.cpp
+++ b/Engine/Source/Programs/Unsync/Private/UnsyncProxy.cpp
@@ -1,19 +1,25 @@
 // Copyright Epic Games, Inc. All Rights Reserved.
 
 #include "UnsyncProxy.h"
+#include "UnsyncAuth.h"
 #include "UnsyncCompression.h"
 #include "UnsyncFile.h"
 #include "UnsyncJupiter.h"
-#include "UnsyncAuth.h"
+#include "UnsyncPool.h"
+#include "UnsyncProgress.h"
 
-#include <json11.hpp>
 #include <fmt/format.h>
+#include <atomic>
+#include <json11.hpp>
 
 namespace unsync {
 
 struct FUnsyncProtocolImpl : FRemoteProtocolBase
 {
-	FUnsyncProtocolImpl(const FRemoteDesc& InRemoteDesc, const FRemoteProtocolFeatures& InFeatures, const FBlockRequestMap* InRequestMap, const FTlsClientSettings* TlsSettings);
+	FUnsyncProtocolImpl(const FRemoteDesc&			   InRemoteDesc,
+						const FRemoteProtocolFeatures& InFeatures,
+						const FBlockRequestMap*		   InRequestMap,
+						const FTlsClientSettings*	   TlsSettings);
 	virtual ~FUnsyncProtocolImpl() override;
 	virtual bool			 IsValid() const override;
 	virtual FDownloadResult	 Download(const TArrayView<FNeedBlock> NeedBlocks, const FBlockDownloadCallback& CompletionCallback) override;
@@ -94,7 +100,8 @@ FUnsyncProtocolImpl::FUnsyncProtocolImpl(const FRemoteDesc&				RemoteDesc,
 
 	if (SocketHandle)
 	{
-		bIsConnetedToHost = [this]() {
+		bIsConnetedToHost = [this]()
+		{
 			FHandshakePacket HandshakePacketTx;
 			if (!SocketSendT(*SocketHandle, HandshakePacketTx))
 			{
@@ -328,7 +335,7 @@ FUnsyncProtocolImpl::Download(const TArrayView<FNeedBlock> NeedBlocks, const FBl
 
 		bOk &= (SocketRecvAll(*SocketHandle, BlockPacket.CompressedData.Data(), BlockPacket.CompressedData.Size()) == CompressedDataSize);
 
-		static const FHash128 TerminatorHash = FHash128{}; // response is always terminated with an empty packet
+		static const FHash128 TerminatorHash = FHash128{};	// response is always terminated with an empty packet
 
 		if (BlockPacket.Hash == TerminatorHash)
 		{
@@ -377,8 +384,8 @@ FUnsyncProtocolImpl::GetSocketSecurity() const
 	}
 }
 
-namespace ProxyQuery {
-TResult<FHelloResponse> Hello(const FRemoteDesc& RemoteDesc, bool bAnonymous)
+TResult<ProxyQuery::FHelloResponse>
+ProxyQuery::Hello(const FRemoteDesc& RemoteDesc, bool bAnonymous)
 {
 	const char* Url = "/api/v1/hello";
 
@@ -391,7 +398,7 @@ TResult<FHelloResponse> Hello(const FRemoteDesc& RemoteDesc, bool bAnonymous)
 			BearerToken = std::move(AuthTokenResult.GetData().Access);
 		}
 	}
-	
+
 	FHttpResponse Response = HttpRequest(RemoteDesc, EHttpMethod::GET, Url, BearerToken);
 
 	if (!Response.Success())
@@ -486,7 +493,8 @@ TResult<FHelloResponse> Hello(const FRemoteDesc& RemoteDesc, bool bAnonymous)
 	return ResultOk(std::move(Result));
 }
 
-TResult<FDirectoryListing> ListDirectory(const FRemoteDesc& Remote, const std::string& Path)
+TResult<ProxyQuery::FDirectoryListing>
+ProxyQuery::ListDirectory(const FRemoteDesc& Remote, const std::string& Path)
 {
 	FDirectoryListing Result;
 
@@ -550,33 +558,161 @@ TResult<FDirectoryListing> ListDirectory(const FRemoteDesc& Remote, const std::s
 	return ResultOk(std::move(Result));
 }
 
-TResult<FBuffer> DownloadFile(const FRemoteDesc& Remote, const std::string& Path)
+TResult<>
+ProxyQuery::DownloadFile(const FRemoteDesc& Remote, const std::string& Path, ProxyQuery::FDownloadOutputCallback OutputCallback)
 {
-	TResult<FAuthToken> AuthToken = Authenticate(Remote, 5 * 60);
-	if (!AuthToken.IsOk())
+	auto CreateConnection = [Remote]
 	{
-		return MoveError<FBuffer>(AuthToken);
-	}
+		FTlsClientSettings TlsSettings;
+		TlsSettings.Subject			   = Remote.HostAddress.data();
+		TlsSettings.bVerifyCertificate = Remote.bTlsVerifyCertificate;
+		return new FHttpConnection(Remote.HostAddress, Remote.HostPort, &TlsSettings);
+	};
 
-	FHttpConnection Connection = FHttpConnection::CreateDefaultHttps(Remote);
+	TObjectPool<FHttpConnection> ConnectionPool(CreateConnection);
 
 	std::string Url = fmt::format("/api/v1/file?{}", Path);
 
-	FHttpRequest Request;
-	Request.Url			   = Url;
-	Request.Method		   = EHttpMethod::GET;
-	Request.BearerToken	   = AuthToken->Access;
-	FHttpResponse Response = HttpRequest(Connection, Request);
+	uint64 FileSize = 0;
 
-	if (!Response.Success())
 	{
-		return HttpError(Response.Code);
+		std::unique_ptr<FHttpConnection> Connection = ConnectionPool.Acquire();
+
+		TResult<FAuthToken> AuthToken = Authenticate(Remote, 5 * 60);
+		if (!AuthToken.IsOk())
+		{
+			return std::move(AuthToken.GetError());
+		}
+
+		FHttpRequest HeadRequest;
+		HeadRequest.Url			   = Url;
+		HeadRequest.Method		   = EHttpMethod::HEAD;
+		HeadRequest.BearerToken	   = AuthToken->Access;
+		FHttpResponse HeadResponse = HttpRequest(*Connection, HeadRequest);
+		if (!HeadResponse.Success())
+		{
+			return HttpError(HeadResponse.Code);
+		}
+
+		FileSize = HeadResponse.ContentLength;
+		UNSYNC_VERBOSE(L"File size: %llu bytes (%.3f MB)", llu(FileSize), SizeMb(FileSize));
+
+		ConnectionPool.Release(std::move(Connection));
+	}
+
+	const uint64 MaxChunkSize = 16_MB;
+
+	FIOWriter& Result = OutputCallback(FileSize);
+	if (!Result.IsValid())
+	{
+		return AppError(L"Failed to create download output stream");
+	}
+
+	const uint64		NumChunks = DivUp(FileSize, MaxChunkSize);
+	std::vector<FRange> Chunks;
+	Chunks.reserve(NumChunks);
+
+	for (uint64 i = 0; i < NumChunks; ++i)
+	{
+		FRange Range;
+		Range.Offset = i * MaxChunkSize;
+		Range.Size	 = CalcChunkSize(i, MaxChunkSize, FileSize);
+		Chunks.push_back(Range);
 	}
 
-	return ResultOk(std::move(Response.Buffer));
+	FAtomicError Error;
+	FSemaphore	 DownloadSempahore(4);	// up to 4 concurrent connections
+
+	FLogProgressScope DownloadProgress(FileSize, ELogProgressUnits::MB);
+
+	auto ProcessChunk =
+		[&Error, &Result, &Url, &Remote, &ConnectionPool, &DownloadSempahore, &DownloadProgress](
+			const FRange& Range)
+	{
+		FLogIndentScope	   IndentScope(DownloadProgress.ParentThreadIndent, true);
+		FLogVerbosityScope VerbosityScope(DownloadProgress.bParentThreadVerbose);
+
+		if (Error.Test())
+		{
+			return;
+		}
+
+		DownloadSempahore.Acquire();
+
+		TResult<FAuthToken> AuthToken = Authenticate(Remote, 5 * 60);
+		if (!AuthToken.IsOk())
+		{
+			Error.Set(std::move(AuthToken.GetError()));
+			return;
+		}
+
+		std::unique_ptr<FHttpConnection> Connection = ConnectionPool.Acquire();
+
+		std::string RequestHeaders = fmt::format("Range: bytes={}-{}", Range.Offset, Range.Offset + Range.Size - 1);
+
+		FHttpRequest Request;
+		Request.Url			   = Url;
+		Request.Method		   = EHttpMethod::GET;
+		Request.BearerToken	   = AuthToken->Access;
+		Request.CustomHeaders  = RequestHeaders;
+		FHttpResponse Response = HttpRequest(*Connection, Request);
+
+		if (!Response.Success())
+		{
+			Error.Set(HttpError(Response.Code));
+			return;
+		}
+
+		if (Range.Size != Response.Buffer.Size())
+		{
+			Error.Set(AppError(L"Downloaded file chunk size mismatch"));
+		}
+
+		uint64 WrittenBytes = Result.Write(Response.Buffer.Data(), Range.Offset, Range.Size);
+
+		DownloadProgress.Add(WrittenBytes);
+
+		ConnectionPool.Release(std::move(Connection));
+
+		DownloadSempahore.Release();
+	};
+
+	ParallelForEach(Chunks, ProcessChunk);
+
+	DownloadProgress.Complete();
+
+	if (Error.Test())
+	{
+		return std::move(*Error.Data);
+	}
+
+	return ResultOk();
 }
 
-}  // namespace ProxyQuery
+TResult<FBuffer>
+ProxyQuery::DownloadFile(const FRemoteDesc& Remote, const std::string& Path)
+{
+	FBuffer Result;
+
+	std::unique_ptr<FMemReaderWriter> ResultWriter;
+
+	auto OutputCallback = [&Result, &ResultWriter](uint64 Size) -> FIOWriter&
+	{
+		Result.Resize(Size);
+		ResultWriter = std::make_unique<FMemReaderWriter>(Result.Data(), Result.Size());
+		return *ResultWriter;
+	};
+
+	TResult<> DownloadResult = DownloadFile(Remote, Path, OutputCallback);
+	if (DownloadResult.IsOk())
+	{
+		return ResultOk(std::move(Result));
+	}
+	else
+	{
+		return MoveError<FBuffer>(DownloadResult);
+	}
+}
 
 void
 FUnsyncProtocolImpl::SendTelemetryEvent(const FRemoteDesc& RemoteDesc, const FTelemetryEventSyncComplete& Event)
@@ -735,8 +871,7 @@ FBlockRequestMap::GetMacroBlockRequest(const FGenericHash& BlockHash) const
 	return Result;
 }
 
-FProxyPool::FProxyPool()
-: FProxyPool(FRemoteDesc())
+FProxyPool::FProxyPool() : FProxyPool(FRemoteDesc())
 {
 }
 
diff --git a/Engine/Source/Programs/Unsync/Private/UnsyncProxy.h b/Engine/Source/Programs/Unsync/Private/UnsyncProxy.h
--- a/Engine/Source/Programs/Unsync/Private/UnsyncProxy.h
+++ b/Engine/Source/Programs/Unsync/Private/UnsyncProxy.h
@@ -222,6 +222,9 @@ struct FDirectoryListing
 TResult<FDirectoryListing> ListDirectory(const FRemoteDesc& Remote, const std::string& Path);
 TResult<FBuffer>		   DownloadFile(const FRemoteDesc& Remote, const std::string& Path);
 
+using FDownloadOutputCallback = std::function<FIOWriter&(uint64 Size)>;
+TResult<> DownloadFile(const FRemoteDesc& Remote, const std::string& Path, FDownloadOutputCallback OutputCallback);
+
 } 
 
 }  // namespace unsync
diff --git a/Engine/Source/Programs/Unsync/Private/UnsyncScavenger.cpp b/Engine/Source/Programs/Unsync/Private/UnsyncScavenger.cpp
--- a/Engine/Source/Programs/Unsync/Private/UnsyncScavenger.cpp
+++ b/Engine/Source/Programs/Unsync/Private/UnsyncScavenger.cpp
@@ -75,7 +75,7 @@ FScavengeDatabase::BuildFromFileSyncTasks(const FSyncDirectoryOptions& SyncOptio
 
 	UNSYNC_VERBOSE(L"Loading scavenged manifests: %llu", llu(Result->Manifests.size()));
 
-	ParallelForEach(Result->Manifests.begin(), Result->Manifests.end(), [](FScavengedManifest& Entry) {
+	ParallelForEach(Result->Manifests, [](FScavengedManifest& Entry) {
 		FLogVerbosityScope VerbosityScope(false);  // turn off logging from threads
 		Entry.bValid = LoadDirectoryManifest(Entry.Manifest, Entry.Root, Entry.ManifestPath);
 		if (Entry.bValid)
diff --git a/Engine/Source/Programs/Unsync/Private/UnsyncThread.h b/Engine/Source/Programs/Unsync/Private/UnsyncThread.h
--- a/Engine/Source/Programs/Unsync/Private/UnsyncThread.h
+++ b/Engine/Source/Programs/Unsync/Private/UnsyncThread.h
@@ -201,4 +201,11 @@ ParallelForEach(IT ItBegin, IT ItEnd, FT F)
 
 #endif	// UNSYNC_USE_CONCRT
 
+template<typename T, typename FT>
+inline void
+ParallelForEach(T Container, FT F)
+{
+	ParallelForEach(std::begin(Container), std::end(Container), F);
+}
+
 }  // namespace unsync
diff --git a/Engine/Source/Programs/Unsync/Private/UnsyncUtil.h b/Engine/Source/Programs/Unsync/Private/UnsyncUtil.h
--- a/Engine/Source/Programs/Unsync/Private/UnsyncUtil.h
+++ b/Engine/Source/Programs/Unsync/Private/UnsyncUtil.h
@@ -74,6 +74,12 @@ MakeU64(uint32 H, uint32 L)
 	return uint64(L) | (uint64(H) << 32);
 }
 
+struct FRange
+{
+	uint64 Offset = 0;
+	uint64 Size = 0;
+};
+
 struct FTimingLogger
 {
 	bool		Enabled	  = false;
