unsync - Fix new ParallelForEach helper

* Pass container in by reference to allow mutation

diff --git a/Engine/Source/Programs/Unsync/Private/UnsyncThread.h b/Engine/Source/Programs/Unsync/Private/UnsyncThread.h
--- a/Engine/Source/Programs/Unsync/Private/UnsyncThread.h
+++ b/Engine/Source/Programs/Unsync/Private/UnsyncThread.h
@@ -203,7 +203,7 @@ ParallelForEach(IT ItBegin, IT ItEnd, FT F)
 
 template<typename T, typename FT>
 inline void
-ParallelForEach(T Container, FT F)
+ParallelForEach(T& Container, FT F)
 {
 	ParallelForEach(std::begin(Container), std::end(Container), F);
 }
