unsync - Add messages in addition to error codes in various places, fix CLI11 non-ascii printing, add --version argument

diff --git a/Engine/Source/Programs/Unsync/Private/UnsyncCore.cpp b/Engine/Source/Programs/Unsync/Private/UnsyncCore.cpp
--- a/Engine/Source/Programs/Unsync/Private/UnsyncCore.cpp
+++ b/Engine/Source/Programs/Unsync/Private/UnsyncCore.cpp
@@ -550,15 +550,15 @@ GetVersionString()
 
 		if (strlen(GitBranch) && strlen(GitRev))
 		{
-			snprintf(Str, sizeof(Str), "v" UNSYNC_VERSION_STR " [%s:%s]", GitBranch, GitRev);
+			snprintf(Str, sizeof(Str), UNSYNC_VERSION_STR " [%s:%s]", GitBranch, GitRev);
 		}
 		else if (strlen(GitRev))
 		{
-			snprintf(Str, sizeof(Str), "v" UNSYNC_VERSION_STR " [%s]", GitRev);
+			snprintf(Str, sizeof(Str), UNSYNC_VERSION_STR " [%s]", GitRev);
 		}
 		else
 		{
-			snprintf(Str, sizeof(Str), "v" UNSYNC_VERSION_STR);
+			snprintf(Str, sizeof(Str), UNSYNC_VERSION_STR);
 		}
 
 		return std::string(Str);
@@ -1304,7 +1304,7 @@ BuildTarget(FIOWriter&			   Output,
 
 		if (DataProvider.GetError())
 		{
-			UNSYNC_ERROR(L"Failed to read blocks from %ls. Error code: %d.", ListName, DataProvider.GetError());
+			UNSYNC_ERROR(L"Failed to read blocks from %ls. %hs", ListName, FormatSystemErrorMessage(DataProvider.GetError()).c_str());
 			bGotError = true;
 			return;
 		}
@@ -1844,9 +1844,9 @@ CreateDirectoryManifest(const FPath& Root, uint32 BlockSize, FAlgorithmOptions A
 			}
 			else
 			{
-				UNSYNC_FATAL(L"Failed to open file '%ls' while computing manifest blocks. System error code: %d.",
+				UNSYNC_FATAL(L"Failed to open file '%ls' while computing manifest blocks. %hs",
 							 FilePath.wstring().c_str(),
-							 File->GetError());
+							 FormatSystemErrorMessage(File->GetError()).c_str());
 			}
 		}
 	}
@@ -1987,9 +1987,9 @@ UpdateDirectoryManifestBlocks(FDirectoryManifest& Result, const FPath& Root, uin
 		}
 		else
 		{
-			UNSYNC_FATAL(L"Failed to open file '%ls' while computing manifest blocks. System error code: %d.",
+			UNSYNC_FATAL(L"Failed to open file '%ls' while computing manifest blocks. %hs",
 						 FilePath.wstring().c_str(),
-						 File->GetError());
+						 FormatSystemErrorMessage(File->GetError()).c_str());
 		}
 	}
 
@@ -2409,9 +2409,9 @@ SyncFile(const FNeedList&		   NeedList,
 			TargetFile = std::make_unique<FNativeFile>(TempTargetFilePath, EFileMode::CreateWriteOnly, TargetFileSizeInfo.TotalBytes);
 			if (TargetFile->GetError() != 0)
 			{
-				UNSYNC_ERROR(L"Failed to create output file '%ls'. Error code %d.",
+				UNSYNC_FATAL(L"Failed to create output file '%ls'. %hs",
 							 TempTargetFilePath.wstring().c_str(),
-							 TargetFile->GetError());
+							 FormatSystemErrorMessage(TargetFile->GetError()).c_str());
 			}
 		}
 
@@ -3868,11 +3868,12 @@ SyncDirectory(const FSyncDirectoryOptions& SyncOptions)
 
 			if (!SyncResult.Succeeded())
 			{
-				UNSYNC_ERROR(L"Sync failed from '%ls' to '%ls'. Status: %ls, system error code: %d",
+				UNSYNC_ERROR(L"Sync failed from '%ls' to '%ls'. Status: %ls, system error code: %d %hs",
 							 Item.ResolvedSourceFilePath.wstring().c_str(),
 							 Item.TargetFilePath.wstring().c_str(),
 							 ToString(SyncResult.Status),
-							 SyncResult.SystemErrorCode.value());
+							 SyncResult.SystemErrorCode.value(),
+							 SyncResult.SystemErrorCode.message().c_str());
 
 				NumFailedTasks++;
 			}
@@ -4013,10 +4014,11 @@ SyncDirectory(const FSyncDirectoryOptions& SyncOptions)
 			{
 				if (!Item.SyncResult.Succeeded())
 				{
-					UNSYNC_ERROR(L"Failed to copy file '%ls' on background task. Status: %ls, system error code: %d",
+					UNSYNC_ERROR(L"Failed to copy file '%ls' on background task. Status: %ls, system error code: %d %hs",
 								 Item.TargetFilePath.wstring().c_str(),
 								 ToString(Item.SyncResult.Status),
-								 Item.SyncResult.SystemErrorCode.value());
+								 Item.SyncResult.SystemErrorCode.value(),
+								 Item.SyncResult.SystemErrorCode.message().c_str());
 				}
 			}
 			UNSYNC_ERROR(L"Background file copy process failed!");
diff --git a/Engine/Source/Programs/Unsync/Private/UnsyncFile.cpp b/Engine/Source/Programs/Unsync/Private/UnsyncFile.cpp
--- a/Engine/Source/Programs/Unsync/Private/UnsyncFile.cpp
+++ b/Engine/Source/Programs/Unsync/Private/UnsyncFile.cpp
@@ -263,10 +263,10 @@ FWindowsFile::CompleteReadCommand(Command& Cmd)
 		}
 		else
 		{
-			UNSYNC_WARNING(L"FNativeFile expected to read %lld bytes, but read %lld. Last error code: %d.",
+			UNSYNC_WARNING(L"FNativeFile expected to read %lld bytes, but read %lld. %hs",
 						   (uint64)ExpectedReadBytes,
 						   (uint64)ReadBytes,
-						   LastError);
+						   FormatSystemErrorMessage(LastError).c_str());
 
 			UNSYNC_LOG(L"Trying to recover from error (attempt %d of %d)", Attempt + 1, MaxAttempts);
 
@@ -288,7 +288,7 @@ FWindowsFile::CompleteReadCommand(Command& Cmd)
 			if (!bOpenedOk)
 			{
 				LastError = GetLastError();
-				UNSYNC_ERROR(L"Failed to re-open the file. Last error: %d.", LastError);
+				UNSYNC_ERROR(L"Failed to re-open the file. %hs", FormatSystemErrorMessage(LastError).c_str());
 				break;
 			}
 
@@ -881,7 +881,9 @@ WriteBufferToFile(const FPath& Filename, const uint8* Data, uint64 Size)
 	}
 	else
 	{
-		UNSYNC_ERROR(L"Failed to open file '%ls' for writing (%d)", Filename.wstring().c_str(), File.GetError());
+		UNSYNC_ERROR(L"Failed to open file '%ls' for writing. %hs",
+					 Filename.wstring().c_str(),
+					 FormatSystemErrorMessage(File.GetError()).c_str());
 		return false;
 	}
 }
diff --git a/Engine/Source/Programs/Unsync/Private/UnsyncLog.cpp b/Engine/Source/Programs/Unsync/Private/UnsyncLog.cpp
--- a/Engine/Source/Programs/Unsync/Private/UnsyncLog.cpp
+++ b/Engine/Source/Programs/Unsync/Private/UnsyncLog.cpp
@@ -144,7 +144,7 @@ FLogFileScope::FLogFileScope(const wchar_t* Filename)
 		LogSetFileInternal(Filename);
 	}
 
-	UNSYNC_VERBOSE2(L"UNSYNC %hs started logging to file '%ls'", GetVersionString().c_str(), Filename);
+	UNSYNC_VERBOSE2(L"UNSYNC v%hs started logging to file '%ls'", GetVersionString().c_str(), Filename);
 }
 
 FLogFileScope::~FLogFileScope()
@@ -407,8 +407,9 @@ LogWriteCrashDump(void* InExceptionPointers)
 		CreateFileW(CrashDumpFilename.c_str(), GENERIC_READ | GENERIC_WRITE, FILE_SHARE_WRITE | FILE_SHARE_READ, 0, CREATE_ALWAYS, 0, 0);
 	if (DumpFile != INVALID_HANDLE_VALUE)
 	{
-		BOOL Ok =
+		const BOOL Ok =
 			MiniDumpWriteDump(GetCurrentProcess(), GetCurrentProcessId(), DumpFile, MiniDumpWithDataSegs, &ExceptionInfo, nullptr, nullptr);
+
 		if (Ok)
 		{
 			LogPrintf(ELogLevel::Info, L"!!! Crash dump file written.\n");
@@ -416,12 +417,12 @@ LogWriteCrashDump(void* InExceptionPointers)
 		}
 		else
 		{
-			LogPrintf(ELogLevel::Error, L"!!! Failed to generate crash dump. Error code: %d.\n", GetLastError());
+			LogPrintf(ELogLevel::Error, L"!!! Failed to generate crash dump. %hs\n", FormatSystemErrorMessage(GetLastError()).c_str());
 		}
 	}
 	else
 	{
-		LogPrintf(ELogLevel::Error, L"!!! Failed to open output file. Error code: %d.\n", GetLastError());
+		LogPrintf(ELogLevel::Error, L"!!! Failed to open output file. %hs\n", FormatSystemErrorMessage(GetLastError()).c_str());
 	}
 
 	LogPrintf(ELogLevel::Error, L"!!! Failed to write dump file.\n");
diff --git a/Engine/Source/Programs/Unsync/Private/UnsyncMain.cpp b/Engine/Source/Programs/Unsync/Private/UnsyncMain.cpp
--- a/Engine/Source/Programs/Unsync/Private/UnsyncMain.cpp
+++ b/Engine/Source/Programs/Unsync/Private/UnsyncMain.cpp
@@ -39,7 +39,7 @@ InnerMain(int Argc, char** Argv)
 {
 	FTimingLogger TimingLogger("Total time");
 
-	std::string AppDescription = "UNSYNC ";
+	std::string AppDescription = "UNSYNC v";
 	AppDescription += GetVersionString();
 	AppDescription +=
 		" -- Differential binary synchronization tool.\n"
@@ -47,6 +47,7 @@ InnerMain(int Argc, char** Argv)
 
 	CLI::App Cli(AppDescription, "unsync");
 	Cli.allow_windows_style_options(true);
+	Cli.set_version_flag("--version", GetVersionString());
 
 	std::vector<CLI::App*> SubCommands;
 
@@ -317,21 +318,24 @@ InnerMain(int Argc, char** Argv)
 
 	// Run the command
 
-#if UNSYNC_PLATFORM_WINDOWS
-	_setmode(_fileno(stdout), _O_TEXT);
-#endif	// UNSYNC_PLATFORM_WINDOWS
-
-	CLI11_PARSE(Cli, Argc, Argv);
+	try
+	{
+		Cli.parse(Argc, Argv);
+	}
+	catch (CLI::Error& E)
+	{
+		std::stringstream OutputStream;
+		const int32		  ReturnCode = Cli.exit(E, OutputStream, OutputStream);
+		std::wstring	  Output	 = ConvertUtf8ToWide(OutputStream.str());
+		wprintf(L"%s", Output.c_str());
+		return ReturnCode;
+	}
 
 	if (Cli.get_subcommands().size() == 0)
 	{
-		printf("%s", Cli.help().c_str());
+		wprintf(L"%hs", Cli.help().c_str());
 	}
 
-#if UNSYNC_PLATFORM_WINDOWS
-	_setmode(_fileno(stdout), _O_U8TEXT);
-#endif	// UNSYNC_PLATFORM_WINDOWS
-
 	// Configure default output mehtod based on subcommand.
 	// In machine-readable mode, all verbose logging is directed to stderr.
 
@@ -340,7 +344,7 @@ InnerMain(int Argc, char** Argv)
 		GLogMachineReadable = true;
 	}
 
-	UNSYNC_VERBOSE(L"UNSYNC %hs", GetVersionString().c_str());
+	UNSYNC_VERBOSE(L"UNSYNC v%hs", GetVersionString().c_str());
 
 	UnsyncMallocInit(bUseDebugMode ? EMallocType::Debug : EMallocType::Default);
 
diff --git a/Engine/Source/Programs/Unsync/Private/UnsyncUtil.cpp b/Engine/Source/Programs/Unsync/Private/UnsyncUtil.cpp
--- a/Engine/Source/Programs/Unsync/Private/UnsyncUtil.cpp
+++ b/Engine/Source/Programs/Unsync/Private/UnsyncUtil.cpp
@@ -23,6 +23,8 @@
 #include <mutex>
 #include <sstream>
 #include <unordered_set>
+#include <system_error>
+#include <fmt/format.h>
 
 namespace unsync {
 
@@ -453,4 +455,11 @@ GetUserHomeDirectory()
 }
 #endif // UNSYNC_PLATFORM_WINDOWS
 
+std::string
+FormatSystemErrorMessage(int32 ErrorCode)
+{
+	std::string ErrorMessage = std::system_category().message(ErrorCode);
+	return fmt::format("Error code {}: {}", ErrorCode, ErrorMessage);
+}
+
 }  // namespace unsync
diff --git a/Engine/Source/Programs/Unsync/Private/UnsyncUtil.h b/Engine/Source/Programs/Unsync/Private/UnsyncUtil.h
--- a/Engine/Source/Programs/Unsync/Private/UnsyncUtil.h
+++ b/Engine/Source/Programs/Unsync/Private/UnsyncUtil.h
@@ -327,4 +327,7 @@ void OpenUrlInDefaultBrowser(const char* Address);
 
 FPath GetUserHomeDirectory();
 
+// Returns string in format 'Error code 123: Some description.`
+std::string FormatSystemErrorMessage(int32 ErrorCode);
+
 }  // namespace unsync
