unsync - Fix various runtime issues on linux

* Wide string formatting
* Path normalization

diff --git a/Engine/Source/Programs/Unsync/Private/UnsyncAuth.cpp b/Engine/Source/Programs/Unsync/Private/UnsyncAuth.cpp
--- a/Engine/Source/Programs/Unsync/Private/UnsyncAuth.cpp
+++ b/Engine/Source/Programs/Unsync/Private/UnsyncAuth.cpp
@@ -812,7 +812,7 @@ Authenticate(const FRemoteDesc& RemoteDesc, int32 RefreshThreshold)
 
 		if (bSaved)
 		{
-			UNSYNC_VERBOSE2(L"Saved authentication token to file: %s", TokenCachePath->wstring().c_str());
+			UNSYNC_VERBOSE2(L"Saved authentication token to file: %ls", TokenCachePath->wstring().c_str());
 		}
 	}
 
diff --git a/Engine/Source/Programs/Unsync/Private/UnsyncCmdQuery.cpp b/Engine/Source/Programs/Unsync/Private/UnsyncCmdQuery.cpp
--- a/Engine/Source/Programs/Unsync/Private/UnsyncCmdQuery.cpp
+++ b/Engine/Source/Programs/Unsync/Private/UnsyncCmdQuery.cpp
@@ -193,7 +193,7 @@ CmdQueryFile(const FCmdQueryOptions& Options)
 
 	if (Response.IsOk())
 	{
-		UNSYNC_LOG(L"Output written to file '%s'", Options.OutputPath.wstring().c_str());
+		UNSYNC_LOG(L"Output written to file '%ls'", Options.OutputPath.wstring().c_str());
 	}
 	else
 	{
diff --git a/Engine/Source/Programs/Unsync/Private/UnsyncCmdSync.cpp b/Engine/Source/Programs/Unsync/Private/UnsyncCmdSync.cpp
--- a/Engine/Source/Programs/Unsync/Private/UnsyncCmdSync.cpp
+++ b/Engine/Source/Programs/Unsync/Private/UnsyncCmdSync.cpp
@@ -41,7 +41,7 @@ CmdSync(const FCmdSyncOptions& Options)
 		UNSYNC_LOG_INDENT;
 		for( const std::wstring& include : Options.Filter->SyncIncludedWords)
 		{ 
-			UNSYNC_VERBOSE(L" %s", include.c_str());
+			UNSYNC_VERBOSE(L" %ls", include.c_str());
 		}
 	}
 	if(!Options.Filter->SyncExcludedWords.empty())
@@ -50,7 +50,7 @@ CmdSync(const FCmdSyncOptions& Options)
 		UNSYNC_LOG_INDENT;
 		for( const std::wstring& exclude : Options.Filter->SyncExcludedWords)
 		{ 
-			UNSYNC_VERBOSE(L"%s", exclude.c_str());
+			UNSYNC_VERBOSE(L"%ls", exclude.c_str());
 		}
 	}
 
diff --git a/Engine/Source/Programs/Unsync/Private/UnsyncCore.cpp b/Engine/Source/Programs/Unsync/Private/UnsyncCore.cpp
--- a/Engine/Source/Programs/Unsync/Private/UnsyncCore.cpp
+++ b/Engine/Source/Programs/Unsync/Private/UnsyncCore.cpp
@@ -4640,7 +4640,7 @@ LogManifestFiles(ELogLevel LogLevel, const FDirectoryManifest& Manifest)
 	for (const std::wstring& Filename : Files)
 	{
 		const FFileManifest& Info = Manifest.Files.at(Filename);
-		LogPrintf(LogLevel, L"%s : %llu\n", Filename.c_str(), Info.Size);
+		LogPrintf(LogLevel, L"%ls : %llu\n", Filename.c_str(), Info.Size);
 	}
 }
 
diff --git a/Engine/Source/Programs/Unsync/Private/UnsyncMain.cpp b/Engine/Source/Programs/Unsync/Private/UnsyncMain.cpp
--- a/Engine/Source/Programs/Unsync/Private/UnsyncMain.cpp
+++ b/Engine/Source/Programs/Unsync/Private/UnsyncMain.cpp
@@ -327,7 +327,7 @@ InnerMain(int Argc, char** Argv)
 		std::stringstream OutputStream;
 		const int32		  ReturnCode = Cli.exit(E, OutputStream, OutputStream);
 		std::wstring	  Output	 = ConvertUtf8ToWide(OutputStream.str());
-		wprintf(L"%s", Output.c_str());
+		wprintf(L"%ls", Output.c_str());
 		return ReturnCode;
 	}
 
diff --git a/Engine/Source/Programs/Unsync/Private/UnsyncScavenger.cpp b/Engine/Source/Programs/Unsync/Private/UnsyncScavenger.cpp
--- a/Engine/Source/Programs/Unsync/Private/UnsyncScavenger.cpp
+++ b/Engine/Source/Programs/Unsync/Private/UnsyncScavenger.cpp
@@ -160,7 +160,7 @@ FScavengeDatabase::BuildFromFileSyncTasks(const FSyncDirectoryOptions& SyncOptio
 
 			FPath FilePath = Manifest.Root / RelativeFilePath;
 
-			UNSYNC_VERBOSE(L"- %s: %llu", FilePath.wstring().c_str(), Item.Count);
+			UNSYNC_VERBOSE(L"- %ls: %llu", FilePath.wstring().c_str(), Item.Count);
 		}
 	}
 #endif
@@ -392,7 +392,7 @@ BuildTargetFromScavengedData(FIOWriter&						Output,
 				continue;
 			}
 
-			UNSYNC_VERBOSE(L"Scavenging data from '%s'", PossibleSource.FullSourceFilePath.wstring().c_str());
+			UNSYNC_VERBOSE(L"Scavenging data from '%ls'", PossibleSource.FullSourceFilePath.wstring().c_str());
 
 			std::sort(LocalNeedList.begin(), LocalNeedList.end(), FNeedBlock::FCompareBySourceOffset());
 
diff --git a/Engine/Source/Programs/Unsync/Private/UnsyncSerialization.cpp b/Engine/Source/Programs/Unsync/Private/UnsyncSerialization.cpp
--- a/Engine/Source/Programs/Unsync/Private/UnsyncSerialization.cpp
+++ b/Engine/Source/Programs/Unsync/Private/UnsyncSerialization.cpp
@@ -654,7 +654,7 @@ SaveDirectoryManifest(const FDirectoryManifest& Manifest, const FPath& Filename)
 	}
 	else
 	{
-		UNSYNC_ERROR(L"Failed to open manifest output file '%ls'", Filename.wstring().c_str());
+		UNSYNC_ERROR(L"Failed to open manifest output file '%ls' for writing", Filename.wstring().c_str());
 		return false;
 	}
 }
diff --git a/Engine/Source/Programs/Unsync/Private/UnsyncUtil.cpp b/Engine/Source/Programs/Unsync/Private/UnsyncUtil.cpp
--- a/Engine/Source/Programs/Unsync/Private/UnsyncUtil.cpp
+++ b/Engine/Source/Programs/Unsync/Private/UnsyncUtil.cpp
@@ -317,8 +317,10 @@ NormalizeFilenameUtf8(const std::string& InFilename)
 	}
 
 	FPath FilenameAsPath = ConvertUtf8ToWide(Filename);
-	FPath NormalPath = FilenameAsPath.lexically_normal();
+
+	FPath NormalPath		 = std::filesystem::weakly_canonical(FilenameAsPath.lexically_normal());
 	FPath AbsoluteNormalPath = std::filesystem::absolute(NormalPath);
+
 	return AbsoluteNormalPath;
 }
 
