unsync - Implement OpenUrlInDefaultBrowser for mac and linux, fix various issues with direct download on those platforms

diff --git a/Engine/Source/Programs/Unsync/Private/UnsyncAuth.cpp b/Engine/Source/Programs/Unsync/Private/UnsyncAuth.cpp
--- a/Engine/Source/Programs/Unsync/Private/UnsyncAuth.cpp
+++ b/Engine/Source/Programs/Unsync/Private/UnsyncAuth.cpp
@@ -337,6 +337,7 @@ AcquireAuthToken(const FAuthDesc& AuthDesc)
 
 	std::thread ServerThread = StartHttpCallbackServer(CallbackListenSocket, CallbackServerDesc.RequestPath, RandomState, HttpCallbackData);
 
+	UNSYNC_LOG(L"Authorization URL: %hs", AuthorizeUrl.c_str());
 	OpenUrlInDefaultBrowser(AuthorizeUrl.c_str());
 
 	UNSYNC_LOG(L"Waiting for HTTP callback on port %d...", int(CallbackPortNumber));
diff --git a/Engine/Source/Programs/Unsync/Private/UnsyncCmdSync.cpp b/Engine/Source/Programs/Unsync/Private/UnsyncCmdSync.cpp
--- a/Engine/Source/Programs/Unsync/Private/UnsyncCmdSync.cpp
+++ b/Engine/Source/Programs/Unsync/Private/UnsyncCmdSync.cpp
@@ -72,6 +72,10 @@ CmdSync(const FCmdSyncOptions& Options)
 			UNSYNC_VERBOSE(L"Server supports access by manifest hash");
 			bSourceFileSystemRequired = false;
 		}
+		else
+		{
+			UNSYNC_VERBOSE2(L"Server does not support direct file access or download by manifest hash. Source file system access is required.");
+		}
 	}
 
 	if (bSourceFileSystemRequired)
@@ -200,7 +204,7 @@ CmdSync(const FCmdSyncOptions& Options)
 		}
 		else
 		{
-			UNSYNC_ERROR(L"Source path does not exist");
+			UNSYNC_ERROR(L"Source path '%ls' does not exist", ResolvedSource.wstring().c_str());
 			return 1;
 		}
 	}
diff --git a/Engine/Source/Programs/Unsync/Private/UnsyncCore.cpp b/Engine/Source/Programs/Unsync/Private/UnsyncCore.cpp
--- a/Engine/Source/Programs/Unsync/Private/UnsyncCore.cpp
+++ b/Engine/Source/Programs/Unsync/Private/UnsyncCore.cpp
@@ -3003,13 +3003,18 @@ LoadAndMergeSourceManifest(FDirectoryManifest& Output,
 						   const FPath&		   SourceManifestOverride,
 						   bool				   bCaseSensitiveTargetFileSystem)
 {
-	auto ResolvePath = [SyncFilter](const FPath& Filename) -> FPath { return SyncFilter ? SyncFilter->Resolve(Filename) : Filename; };
-
 	const FRemoteProtocolFeatures& ProxyFeatures = ProxyPool.GetFeatures();
 
 	const bool bDownloadManifestFromProxy =
 		ProxyPool.RemoteDesc.Protocol == EProtocolFlavor::Unsync && ProxyFeatures.bDirectoryListing && ProxyFeatures.bFileDownload;
 
+	UNSYNC_VERBOSE2(L"LoadAndMergeSourceManifest: '%ls' (%hs)",
+					SourcePath.wstring().c_str(),
+					bDownloadManifestFromProxy ? "download" : "filesystem");
+
+	auto ResolvePath = [SyncFilter, bDownloadManifestFromProxy](const FPath& Filename) -> FPath
+	{ return (SyncFilter && !bDownloadManifestFromProxy) ? SyncFilter->Resolve(Filename) : Filename; };
+
 	FDirectoryManifest LoadedManifest;
 
 	FPath SourceManifestRoot = SourcePath / ".unsync";
diff --git a/Engine/Source/Programs/Unsync/Private/UnsyncUtil.cpp b/Engine/Source/Programs/Unsync/Private/UnsyncUtil.cpp
--- a/Engine/Source/Programs/Unsync/Private/UnsyncUtil.cpp
+++ b/Engine/Source/Programs/Unsync/Private/UnsyncUtil.cpp
@@ -110,8 +110,8 @@ IsTrivialAsciiString(const T& Input)
 
 #ifdef __clang__
 #	pragma clang diagnostic push
-#	pragma clang diagnostic ignored "-Wdeprecated-declarations" // codecvt_utf8 is deprecated, but there is no trivial replacement
-#endif	// __clang__
+#	pragma clang diagnostic ignored "-Wdeprecated-declarations"  // codecvt_utf8 is deprecated, but there is no trivial replacement
+#endif															  // __clang__
 
 std::wstring
 ConvertUtf8ToWide(std::string_view StringUtf8)
@@ -280,7 +280,8 @@ DfsEnumerate(const FPath& Root)
 	return Result;
 }
 
-FPath GetUniversalPath(const FPath& Path)
+FPath
+GetUniversalPath(const FPath& Path)
 {
 	FPath Result = Path;
 
@@ -292,7 +293,7 @@ FPath GetUniversalPath(const FPath& Path)
 
 	WCHAR Buffer[MaxBufferSize] = {};
 
-	DWORD				 BufferSize		   = MaxBufferSize;
+	DWORD				  BufferSize		= MaxBufferSize;
 	UNIVERSAL_NAME_INFOW* UniversalNameInfo = (UNIVERSAL_NAME_INFOW*)Buffer;
 
 	DWORD ErrorCode = WNetGetUniversalNameW(Path.native().c_str(), UNIVERSAL_NAME_INFO_LEVEL, (LPVOID)UniversalNameInfo, &BufferSize);
@@ -309,6 +310,11 @@ FPath GetUniversalPath(const FPath& Path)
 FPath
 NormalizeFilenameUtf8(const std::string& InFilename)
 {
+	if (InFilename.empty())
+	{
+		return FPath();
+	}
+
 	std::string_view Filename	   = InFilename;
 	std::string_view FileUrlPrefix = "file://";
 	if (Filename.starts_with(FileUrlPrefix))
@@ -318,8 +324,17 @@ NormalizeFilenameUtf8(const std::string& InFilename)
 
 	FPath FilenameAsPath = ConvertUtf8ToWide(Filename);
 
-	FPath NormalPath		 = std::filesystem::weakly_canonical(FilenameAsPath.lexically_normal());
-	FPath AbsoluteNormalPath = std::filesystem::absolute(NormalPath);
+	FPath NormalPath = FilenameAsPath.lexically_normal();
+	FPath AbsoluteNormalPath;
+	if (Filename.starts_with("\\\\") || Filename.starts_with("//"))
+	{
+		AbsoluteNormalPath = NormalPath;  // Assume network paths are absolute
+	}
+	else
+	{
+		FPath CanonicalPath = std::filesystem::weakly_canonical(NormalPath);
+		AbsoluteNormalPath	= std::filesystem::absolute(CanonicalPath);
+	}
 
 	return AbsoluteNormalPath;
 }
@@ -327,12 +342,11 @@ NormalizeFilenameUtf8(const std::string& InFilename)
 FPath
 GetAbsoluteNormalPath(const FPath& InPath)
 {
-	FPath NormalPath = InPath.lexically_normal();
+	FPath NormalPath		 = InPath.lexically_normal();
 	FPath AbsoluteNormalPath = std::filesystem::absolute(NormalPath);
 	return AbsoluteNormalPath;
 }
 
-
 const FBuffer&
 GetSystemRootCerts()
 {
@@ -426,15 +440,26 @@ GetSystemRootCerts()
 }
 
 #if UNSYNC_PLATFORM_WINDOWS
-void OpenUrlInDefaultBrowser(const char* Address)
+void
+OpenUrlInDefaultBrowser(const char* Address)
 {
 	ShellExecuteA(nullptr, "open", Address, nullptr, nullptr, SW_SHOWNORMAL);
 }
-#else // UNSYNC_PLATFORM_WINDOWS
+#else  // UNSYNC_PLATFORM_WINDOWS
 void
 OpenUrlInDefaultBrowser(const char* Address)
 {
-	UNSYNC_FATAL(L"OpenUrlInDefaultBrowser is not implemented");
+#	ifdef __APPLE__
+	std::string Command = fmt::format("open \"{}\"", Address);
+#	else // assume linux with xdg-utils installed
+	std::string Command = fmt::format("xdg-open \"{}\"", Address);
+#	endif
+
+	int RetCode = system(Command.c_str());
+	if (RetCode != 0)
+	{
+		UNSYNC_ERROR(L"Failed to run command '%hs'", Command.c_str());
+	}
 }
 #endif // UNSYNC_PLATFORM_WINDOWS
 
