unsync - Set a 60 second timeout on socket recv operations by default

* Use --no-timeout flag to disable

diff --git a/Engine/Source/Programs/Unsync/Private/UnsyncLog.cpp b/Engine/Source/Programs/Unsync/Private/UnsyncLog.cpp
--- a/Engine/Source/Programs/Unsync/Private/UnsyncLog.cpp
+++ b/Engine/Source/Programs/Unsync/Private/UnsyncLog.cpp
@@ -218,13 +218,19 @@ LogPrintf(ELogLevel Level, const wchar_t* Str, ...)
 	bool		   bShouldIndent = false;
 	const wchar_t* Prefix		 = nullptr;
 
+	const uint32 ThreadIndex = GetLogThreadIndex();
+
+	bool bShouldOutputThreadIndex = false;
+
 	if (Level == ELogLevel::Error)
 	{
 		Prefix = L"ERROR: ";
+		bShouldOutputThreadIndex = true;
 	}
 	else if (Level == ELogLevel::Warning)
 	{
 		Prefix = L"WARNING: ";
+		bShouldOutputThreadIndex = true;
 	}
 	else if (GLogIndent)
 	{
@@ -255,6 +261,11 @@ LogPrintf(ELogLevel Level, const wchar_t* Str, ...)
 			fwprintf(LogStream, Prefix);
 		}
 
+		if (bShouldOutputThreadIndex)
+		{
+			fwprintf(LogStream, L"[Thread %d] ", ThreadIndex);
+		}
+
 		if (bShouldIndent)
 		{
 			fwprintf(LogStream, L"%*c", GLogIndent, L' ');
@@ -272,7 +283,6 @@ LogPrintf(ELogLevel Level, const wchar_t* Str, ...)
 	{
 		FILE* LogFileStream = GLogFile->Handle;
 
-		uint32 ThreadIndex = GetLogThreadIndex();
 		fwprintf(LogFileStream, L"[%3d] ", ThreadIndex);
 
 		switch (Level)
diff --git a/Engine/Source/Programs/Unsync/Private/UnsyncMain.cpp b/Engine/Source/Programs/Unsync/Private/UnsyncMain.cpp
--- a/Engine/Source/Programs/Unsync/Private/UnsyncMain.cpp
+++ b/Engine/Source/Programs/Unsync/Private/UnsyncMain.cpp
@@ -94,6 +94,7 @@ InnerMain(int Argc, char** Argv)
 	bool					 bShouldLogin		 = false;
 	bool					 bQuickLogin		 = false;
 	bool					 bForceRefreshAuth	 = false;
+	bool					 bNoSocketTimeout	 = false;
 	int32					 CompressionLevel	 = 3;
 	uint32					 DiffBlockSize		 = uint32(4_KB);
 	uint32					 HashOrSyncBlockSize = uint32(64_KB);
@@ -257,6 +258,7 @@ InnerMain(int Argc, char** Argv)
 	SubSync->add_option("-b, --block", HashOrSyncBlockSize, "Block size in bytes (default=64KB)");
 	SubSync->add_option("--scavenge", ScavengeRootUtf8, "Search for unsync manifests and reusable blocks in this directory (EXPERIMENTAL)");
 	SubSync->add_flag("--login", bShouldLogin, "Use user authentication when accessing unsync server");
+	SubSync->add_flag("--no-timeout", bNoSocketTimeout, "Disable the default 60 second timeout on network socket operations");
 
 	SubCommands.push_back(SubSync);
 
@@ -768,6 +770,15 @@ InnerMain(int Argc, char** Argv)
 			}
 		}
 
+		if (bNoSocketTimeout)
+		{
+			RemoteDesc.RecvTimeoutSeconds = 0;
+		}
+		else
+		{
+			RemoteDesc.RecvTimeoutSeconds = 60;
+		}
+
 		FCmdSyncOptions SyncOptions;
 
 		SyncOptions.Algorithm			   = Algorithm;
diff --git a/Engine/Source/Programs/Unsync/Private/UnsyncProxy.cpp b/Engine/Source/Programs/Unsync/Private/UnsyncProxy.cpp
--- a/Engine/Source/Programs/Unsync/Private/UnsyncProxy.cpp
+++ b/Engine/Source/Programs/Unsync/Private/UnsyncProxy.cpp
@@ -78,6 +78,8 @@ FUnsyncProtocolImpl::FUnsyncProtocolImpl(const FRemoteDesc&				RemoteDesc,
 	if (RemoteDesc.bTlsEnable && TlsSettings)
 	{
 		FSocketHandle RawSocketHandle = SocketConnectTcp(RemoteDesc.HostAddress.c_str(), RemoteDesc.HostPort);
+		SocketSetRecvTimeout(RawSocketHandle, RemoteDesc.RecvTimeoutSeconds);
+
 		if (RawSocketHandle)
 		{
 			FSocketTls* TlsSocket = new FSocketTls(RawSocketHandle, *TlsSettings);
@@ -95,7 +97,9 @@ FUnsyncProtocolImpl::FUnsyncProtocolImpl(const FRemoteDesc&				RemoteDesc,
 	if (!SocketHandle)
 	{
 		FSocketHandle RawSocketHandle = SocketConnectTcp(RemoteDesc.HostAddress.c_str(), RemoteDesc.HostPort);
-		SocketHandle				  = std::unique_ptr<FSocketRaw>(new FSocketRaw(RawSocketHandle));
+		SocketSetRecvTimeout(RawSocketHandle, RemoteDesc.RecvTimeoutSeconds);
+
+		SocketHandle = std::unique_ptr<FSocketRaw>(new FSocketRaw(RawSocketHandle));
 	}
 
 	if (SocketHandle)
@@ -331,10 +335,22 @@ FUnsyncProtocolImpl::Download(const TArrayView<FNeedBlock> NeedBlocks, const FBl
 		uint64 CompressedDataSize = 0;
 		bOk &= SocketRecvT(*SocketHandle, CompressedDataSize);
 
+		if (!bOk)
+		{
+			UNSYNC_WARNING(L"Failed to receive block header");
+			break;
+		}
+
 		BlockPacket.CompressedData.Resize(CompressedDataSize);
 
 		bOk &= (SocketRecvAll(*SocketHandle, BlockPacket.CompressedData.Data(), BlockPacket.CompressedData.Size()) == CompressedDataSize);
 
+		if (!bOk)
+		{
+			UNSYNC_WARNING(L"Failed to receive block buffer");
+			break;
+		}
+
 		static const FHash128 TerminatorHash = FHash128{};	// response is always terminated with an empty packet
 
 		if (BlockPacket.Hash == TerminatorHash)
diff --git a/Engine/Source/Programs/Unsync/Private/UnsyncRemote.h b/Engine/Source/Programs/Unsync/Private/UnsyncRemote.h
--- a/Engine/Source/Programs/Unsync/Private/UnsyncRemote.h
+++ b/Engine/Source/Programs/Unsync/Private/UnsyncRemote.h
@@ -48,6 +48,8 @@ struct FRemoteDesc
 
 	bool bAuthenticationRequired = false;
 
+	uint32 RecvTimeoutSeconds = 0;
+
 	uint32 MaxConnections = 8;	// Limit on concurrent connections to this server
 
 	bool IsValid() const { return Protocol != EProtocolFlavor::Unknown && !HostAddress.empty() && HostPort != 0; }
diff --git a/Engine/Source/Programs/Unsync/Private/UnsyncSocket.cpp b/Engine/Source/Programs/Unsync/Private/UnsyncSocket.cpp
--- a/Engine/Source/Programs/Unsync/Private/UnsyncSocket.cpp
+++ b/Engine/Source/Programs/Unsync/Private/UnsyncSocket.cpp
@@ -235,6 +235,8 @@ SocketSend(FSocketHandle Socket, const void* Data, size_t DataSize)
 	if (DataSize == 0)
 		return 0;
 
+	UNSYNC_ASSERT(DataSize < INT_MAX);
+
 	int Result = send(Socket, reinterpret_cast<const char*>(Data), (int)DataSize, 0);
 	return Result;
 }
@@ -245,7 +247,25 @@ SocketRecvAll(FSocketHandle Socket, void* Data, size_t DataSize)
 	if (DataSize == 0)
 		return 0;
 
-	return recv(Socket, reinterpret_cast<char*>(Data), (int)DataSize, MSG_WAITALL);
+	UNSYNC_ASSERT(DataSize < INT_MAX);
+
+	// If socket timeout is set, recv may read partial data and return the
+	// number of bytes read before timeout or <=0 if another error occurred.
+
+	int ProcessedBytes = 0;
+	while (ProcessedBytes < int(DataSize))
+	{
+		char* BatchPtr		= reinterpret_cast<char*>(Data) + ProcessedBytes;
+		int	  BatchSize		= (int)DataSize - ProcessedBytes;
+		int	  Res = recv(Socket, BatchPtr, BatchSize, MSG_WAITALL);
+		if (Res <= 0)
+		{
+			break;
+		}
+		ProcessedBytes += Res;
+	}
+
+	return ProcessedBytes;
 }
 
 int
@@ -254,9 +274,27 @@ SocketRecvAny(FSocketHandle Socket, void* Data, size_t DataSize)
 	if (DataSize == 0)
 		return 0;
 
+	UNSYNC_ASSERT(DataSize < INT_MAX);
+
 	return recv(Socket, reinterpret_cast<char*>(Data), (int)DataSize, 0);
 }
 
+bool
+SocketSetRecvTimeout(FSocketHandle Socket, uint32 Seconds)
+{
+#if UNSYNC_PLATFORM_WINDOWS
+	uint32 Timeout = Seconds * 1000;  // Winsock uses timeout in milliseconds
+#else
+	struct timeval Timeout;
+	Timeout.tv_sec	= long(Seconds);
+	Timeout.tv_usec = 0;
+#endif
+
+	int Result = setsockopt(Socket, SOL_SOCKET, SO_RCVTIMEO, reinterpret_cast<const char*>(&Timeout), int(sizeof(Timeout)));
+
+	return Result == 0;
+}
+
 bool
 SocketValid(FSocketHandle Socket)
 {
@@ -387,6 +425,8 @@ FSocketTls::Send(const void* Data, size_t DataSize)
 	if (DataSize == 0)
 		return 0;
 
+	UNSYNC_ASSERT(DataSize < INT_MAX);
+
 	int ProcessedBytes = 0;
 	while (ProcessedBytes < DataSize)
 	{
@@ -407,6 +447,8 @@ FSocketTls::RecvAll(void* Data, size_t DataSize)
 	if (DataSize == 0)
 		return 0;
 
+	UNSYNC_ASSERT(DataSize < INT_MAX);
+
 	int ProcessedBytes = 0;
 	while (ProcessedBytes < DataSize)
 	{
@@ -427,6 +469,8 @@ FSocketTls::RecvAny(void* Data, size_t DataSize)
 	if (DataSize == 0)
 		return 0;
 
+	UNSYNC_ASSERT(DataSize < INT_MAX);
+
 	int ProcessedBytes = 0;
 	do
 	{
diff --git a/Engine/Source/Programs/Unsync/Private/UnsyncSocket.h b/Engine/Source/Programs/Unsync/Private/UnsyncSocket.h
--- a/Engine/Source/Programs/Unsync/Private/UnsyncSocket.h
+++ b/Engine/Source/Programs/Unsync/Private/UnsyncSocket.h
@@ -47,6 +47,8 @@ int			   SocketRecvAll(FSocketHandle Socket, void* Data, size_t DataSize);
 int			   SocketRecvAny(FSocketHandle Socket, void* Data, size_t DataSize);
 FSocketAddress SocketMakeAddress(const char* Address, uint16 Port);
 
+bool SocketSetRecvTimeout(FSocketHandle Socket, uint32 Seconds);
+
 bool SocketValid(FSocketHandle Socket);
 
 template<typename T>
