unsync - Add timestamps to log files

diff --git a/Engine/Source/Programs/Unsync/Private/UnsyncLog.cpp b/Engine/Source/Programs/Unsync/Private/UnsyncLog.cpp
--- a/Engine/Source/Programs/Unsync/Private/UnsyncLog.cpp
+++ b/Engine/Source/Programs/Unsync/Private/UnsyncLog.cpp
@@ -12,6 +12,10 @@
 #include <atomic>
 #include <mutex>
 #include <vector>
+#include <chrono>
+
+#include <fmt/format.h>
+#include <fmt/chrono.h>
 
 #if UNSYNC_PLATFORM_WINDOWS
 #	include <Dbghelp.h>
@@ -210,11 +214,34 @@ LogFlush()
 	}
 }
 
+static constexpr size_t TimestampStringSize = sizeof("YYYY-MM-DDTHH:MM:SS.sssZ");
+
+using FTimestampString = fmt::basic_memory_buffer<char, TimestampStringSize>;
+
+static FTimestampString
+FormatTimestamp(std::chrono::system_clock::time_point Timestamp)
+{
+	auto	TimestampMilliseconds = std::chrono::time_point_cast<std::chrono::milliseconds>(Timestamp);
+	auto	TimestampSeconds	  = std::chrono::time_point_cast<std::chrono::seconds>(TimestampMilliseconds);
+	auto	DeltaMilliseconds	  = (TimestampMilliseconds - TimestampSeconds).count();
+	std::tm UtcTime				  = fmt::gmtime(std::chrono::system_clock::to_time_t(TimestampSeconds));
+
+	FTimestampString Result;
+	fmt::format_to(std::back_inserter(Result), "{:%FT%T}.{:03}Z", UtcTime, DeltaMilliseconds);
+	Result.push_back(0);
+
+	UNSYNC_ASSERT(Result.size() == TimestampStringSize);
+
+	return Result;
+}
+
 void
 LogPrintf(ELogLevel Level, const wchar_t* Str, ...)
 {
 	std::lock_guard<std::mutex> LockGuard(GLogMutex);
 
+	const auto Timestamp = std::chrono::system_clock::now();
+
 	bool		   bShouldIndent = false;
 	const wchar_t* Prefix		 = nullptr;
 
@@ -283,6 +310,10 @@ LogPrintf(ELogLevel Level, const wchar_t* Str, ...)
 	{
 		FILE* LogFileStream = GLogFile->Handle;
 
+		FTimestampString TimestampString = FormatTimestamp(Timestamp);
+
+		fwprintf(LogFileStream, L"[%hs] ", TimestampString.data());
+
 		fwprintf(LogFileStream, L"[%3d] ", ThreadIndex);
 
 		switch (Level)
