unsync - Improve error reporting during file patching

* Bump version to 1.0.56

diff --git a/Engine/Source/Programs/Unsync/Private/UnsyncCore.cpp b/Engine/Source/Programs/Unsync/Private/UnsyncCore.cpp
--- a/Engine/Source/Programs/Unsync/Private/UnsyncCore.cpp
+++ b/Engine/Source/Programs/Unsync/Private/UnsyncCore.cpp
@@ -27,7 +27,7 @@ UNSYNC_THIRD_PARTY_INCLUDES_START
 #include <md5-sse2.h>
 UNSYNC_THIRD_PARTY_INCLUDES_END
 
-#define UNSYNC_VERSION_STR "1.0.55"
+#define UNSYNC_VERSION_STR "1.0.56"
 
 namespace unsync {
 
@@ -1232,9 +1232,16 @@ BuildTarget(FIOWriter&			   Output,
 	auto TimeBegin = TimePointNow();
 
 	const FNeedListSize SizeInfo = ComputeNeedListSize(NeedList);
+
+	if (SizeInfo.TotalBytes != Output.GetSize())
+	{
+		UNSYNC_ERROR(L"Output size is %llu, but expected to be %llu", llu(Output.GetSize()), llu(SizeInfo.TotalBytes));
+		return BuildResult;
+	}
+
 	UNSYNC_ASSERT(SizeInfo.TotalBytes == Output.GetSize());
 
-	std::atomic<bool> bGotError = false;
+	FAtomicError Error;
 
 	std::atomic<bool> bBaseDataCopyTaskDone	  = false;
 	std::atomic<bool> bSourceDataCopyTaskDone = false;
@@ -1284,7 +1291,8 @@ BuildTarget(FIOWriter&			   Output,
 		FilteredSourceNeedList = NeedList.Source;
 	}
 
-	auto ProcessNeedList = [bAllowVerboseLog, LogIndent, &Output, &bGotError, &WriteSemaphore, &WriteTasks, &bWaitingForBaseData, &Stats](
+	auto ProcessNeedList =
+		[bAllowVerboseLog, LogIndent, &Output, &Error, &WriteSemaphore, &WriteTasks, &bWaitingForBaseData, &Stats, SizeInfo](
 							   FIOReader&					  DataProvider,
 							   const std::vector<FNeedBlock>& NeedBlocks,
 							   uint64						  TotalCopySize,
@@ -1298,14 +1306,23 @@ BuildTarget(FIOWriter&			   Output,
 		if (!DataProvider.IsValid())
 		{
 			UNSYNC_ERROR(L"Failed to read blocks from %ls. Stream is invalid.", ListName);
-			bGotError = true;
+			Error.Set(AppError(L"Failed to read blocks. Stream is invalid."));
+			return;
+		}
+
+		if (ListType == EBlockListType::Source && DataProvider.GetSize() != SizeInfo.TotalBytes)
+		{
+			UNSYNC_ERROR(L"File size is %llu, but expected to be %llu. File may have changed after manifest was generated.",
+						 llu(DataProvider.GetSize()),
+						 llu(SizeInfo.TotalBytes));
+			Error.Set(AppError(L"Failed to read source blocks. Size mismatch."));
 			return;
 		}
 
 		if (DataProvider.GetError())
 		{
 			UNSYNC_ERROR(L"Failed to read blocks from %ls. %hs", ListName, FormatSystemErrorMessage(DataProvider.GetError()).c_str());
-			bGotError = true;
+			Error.Set(AppError(L"Failed to read blocks"));
 			return;
 		}
 
@@ -1324,7 +1341,7 @@ BuildTarget(FIOWriter&			   Output,
 		static constexpr uint32 MaxActiveLargeRequests							= 2;
 		uint64					ActiveLargeRequestSizes[MaxActiveLargeRequests] = {};
 
-		while (!ReadSchedule.Requests.empty() && !bGotError)
+		while (!ReadSchedule.Requests.empty() && !Error)
 		{
 			uint64 BlockIndex = ~0ull;
 
@@ -1356,12 +1373,17 @@ BuildTarget(FIOWriter&			   Output,
 				ActiveLargeRequestSizes[LargeRequestSlot] = Block.Size;
 			}
 
-			// UNSYNC_VERBOSE(L"Reading block %d, size %d", (int)block_index, (int)block.size);
+			UNSYNC_ASSERTF(Block.SourceOffset + Block.Size <= DataProvider.GetSize(),
+						   L"Copy command is out of bounds. Offset %llu, size %llu ([%llu..%llu]), input size %llu.",
+						   llu(Block.SourceOffset),
+						   llu(Block.Size),
+						   llu(Block.SourceOffset),
+						   llu(Block.SourceOffset + Block.Size),
+						   DataProvider.GetSize());
 
-			UNSYNC_ASSERT(Block.SourceOffset + Block.Size <= DataProvider.GetSize());
 			uint64 ReadBytes = 0;
 
-			auto ReadCallback = [&ReadBytes, &Output, &WriteTasks, &bGotError, &WriteSemaphore, &Stats, Block, ListType](
+			auto ReadCallback = [&ReadBytes, &Output, &WriteTasks, &Error, &WriteSemaphore, &Stats, Block, ListType](
 									FIOBuffer CmdBuffer,
 									uint64	  CmdOffset,
 									uint64	  CmdReadSize,
@@ -1370,8 +1392,7 @@ BuildTarget(FIOWriter&			   Output,
 				WriteTasks.run([Buffer = MakeShared(std::move(CmdBuffer)),
 								CmdReadSize,
 								Block,
-								&Output,
-								&bGotError,
+								&Output, &Error,
 								&WriteSemaphore,
 								&Stats,
 								ListType]() {
@@ -1395,7 +1416,7 @@ BuildTarget(FIOWriter&			   Output,
 					if (WrittenBytes != CmdReadSize)
 					{
 						UNSYNC_FATAL(L"Expected to write %llu bytes, but written %llu", CmdReadSize, WrittenBytes);
-						bGotError = true;
+						Error.Set(AppError(L"Failed to write output"));
 					}
 				});
 
@@ -1541,7 +1562,7 @@ BuildTarget(FIOWriter&			   Output,
 			 &BlockScatterMap,
 			 &DownloadedBlocksMutex,
 			 &DownloadedBlocks,
-			 &bGotError,
+			 &Error,
 			 &NumHashMismatches,
 			 &ParentThreadIndent,
 			 &bParentThreadVerbose,
@@ -1579,7 +1600,7 @@ BuildTarget(FIOWriter&			   Output,
 										 &DownloadedBlocksMutex,
 										 &DownloadedBlocks,
 										 &DecompressionSemaphore,
-										 &bGotError,
+										 &Error,
 										 &NumHashMismatches,
 										 &Stats]() {
 						FLogIndentScope IndentScope(ParentThreadIndent, true);
@@ -1623,7 +1644,7 @@ BuildTarget(FIOWriter&			   Output,
 									{
 										bOk = false;
 										UNSYNC_FATAL(L"Expected to write %llu bytes, but written %llu", Cmd.Size, WrittenBytes);
-										bGotError = true;
+										Error.Set(AppError(L"Failed to write output"));
 									}
 								}
 								else
@@ -1634,7 +1655,7 @@ BuildTarget(FIOWriter&			   Output,
 							else
 							{
 								UNSYNC_FATAL(L"Failed to decompress downloaded block");
-								bGotError = true;
+								Error.Set(AppError(L"Failed to decompress downloaded block"));
 							}
 						}
 
@@ -1653,7 +1674,7 @@ BuildTarget(FIOWriter&			   Output,
 									if (WrittenBytes != Cmd.Size)
 									{
 										UNSYNC_FATAL(L"Expected to write %llu bytes, but written %llu", Cmd.Size, WrittenBytes);
-										bGotError = true;
+										Error.Set(AppError(L"Failed to write output"));
 									}
 								}
 
@@ -1729,7 +1750,7 @@ BuildTarget(FIOWriter&			   Output,
 		LogGlobalProgress();
 	}
 
-	BuildResult.bSuccess	= bGotError.load();
+	BuildResult.bSuccess	= !Error;
 	BuildResult.BaseBytes	= Stats.WrittenBytesFromBase;
 	BuildResult.SourceBytes = Stats.WrittenBytesFromSource;
 
@@ -2435,6 +2456,12 @@ SyncFile(const FNeedList&		   NeedList,
 		Result.SourceBytes = BuildResult.SourceBytes;
 		Result.BaseBytes   = BuildResult.BaseBytes;
 
+		if (!BuildResult.bSuccess)
+		{
+			Result.Status = EFileSyncStatus::ErrorBuildTargetFailed;
+			return Result;
+		}
+
 		if (Options.bValidateTargetFiles)
 		{
 			LogStatus(TargetFilePath.wstring().c_str(), L"Verifying");
@@ -3850,35 +3877,47 @@ SyncDirectory(const FSyncDirectoryOptions& SyncOptions)
 
 			LogStatus(Item.TargetFilePath.wstring().c_str(), SyncResult.Succeeded() ? L"Succeeded" : L"Failed");
 
-			StatSourceBytes += SyncResult.SourceBytes;
-			StatBaseBytes += SyncResult.BaseBytes;
-			UNSYNC_ASSERT(SyncResult.SourceBytes + SyncResult.BaseBytes == Item.TotalSizeBytes);
-
-			if (SyncResult.Succeeded() && !GDryRun)
+			if (SyncResult.Succeeded())
 			{
-				BaseFile = nullptr;
-				SetFileMtime(Item.TargetFilePath, Item.SourceManifest->Mtime);
-			}
+				StatSourceBytes += SyncResult.SourceBytes;
+				StatBaseBytes += SyncResult.BaseBytes;
+				UNSYNC_ASSERT(SyncResult.SourceBytes + SyncResult.BaseBytes == Item.TotalSizeBytes);
 
-			if (bBackground)
-			{
-				FBackgroundTaskResult Result;
-				Result.TargetFilePath = Item.TargetFilePath;
-				Result.SyncResult	  = SyncResult;
-				Result.bIsPartialCopy = Item.NeedBytesFromBase != 0;
+				if (!GDryRun)
+				{
+					BaseFile = nullptr;
+					SetFileMtime(Item.TargetFilePath, Item.SourceManifest->Mtime);
+				}
 
-				std::lock_guard<std::mutex> LockGuard(BackgroundTaskStatMutex);
-				BackgroundTaskResults.push_back(Result);
+				if (bBackground)
+				{
+					FBackgroundTaskResult Result;
+					Result.TargetFilePath = Item.TargetFilePath;
+					Result.SyncResult	  = SyncResult;
+					Result.bIsPartialCopy = Item.NeedBytesFromBase != 0;
+
+					std::lock_guard<std::mutex> LockGuard(BackgroundTaskStatMutex);
+					BackgroundTaskResults.push_back(Result);
+				}	
 			}
-
-			if (!SyncResult.Succeeded())
+			else
 			{
-				UNSYNC_ERROR(L"Sync failed from '%ls' to '%ls'. Status: %ls, system error code: %d %hs",
-							 Item.ResolvedSourceFilePath.wstring().c_str(),
-							 Item.TargetFilePath.wstring().c_str(),
-							 ToString(SyncResult.Status),
-							 SyncResult.SystemErrorCode.value(),
-							 SyncResult.SystemErrorCode.message().c_str());
+				if (SyncResult.SystemErrorCode.value())
+				{
+					UNSYNC_ERROR(L"Sync failed from '%ls' to '%ls'. Status: %ls, system error code: %d %hs",
+								 Item.ResolvedSourceFilePath.wstring().c_str(),
+								 Item.TargetFilePath.wstring().c_str(),
+								 ToString(SyncResult.Status),
+								 SyncResult.SystemErrorCode.value(),
+								 SyncResult.SystemErrorCode.message().c_str());
+				}
+				else
+				{
+					UNSYNC_ERROR(L"Sync failed from '%ls' to '%ls'. Status: %ls.",
+								 Item.ResolvedSourceFilePath.wstring().c_str(),
+								 Item.TargetFilePath.wstring().c_str(),
+								 ToString(SyncResult.Status));
+				}
 
 				NumFailedTasks++;
 			}
@@ -4558,6 +4597,8 @@ ToString(EFileSyncStatus Status)
 			return L"Final file rename failed";
 		case EFileSyncStatus::ErrorTargetFileCreate:
 			return L"Target file creation failed";
+		case EFileSyncStatus::ErrorBuildTargetFailed:
+			return L"Failed to build target";
 	}
 }
 
diff --git a/Engine/Source/Programs/Unsync/Private/UnsyncCore.h b/Engine/Source/Programs/Unsync/Private/UnsyncCore.h
--- a/Engine/Source/Programs/Unsync/Private/UnsyncCore.h
+++ b/Engine/Source/Programs/Unsync/Private/UnsyncCore.h
@@ -300,6 +300,7 @@ enum class EFileSyncStatus
 	ErrorValidation,
 	ErrorFinalRename,
 	ErrorTargetFileCreate,
+	ErrorBuildTargetFailed,
 };
 
 const wchar_t* ToString(EFileSyncStatus Status);
diff --git a/Engine/Source/Programs/Unsync/Private/UnsyncError.h b/Engine/Source/Programs/Unsync/Private/UnsyncError.h
--- a/Engine/Source/Programs/Unsync/Private/UnsyncError.h
+++ b/Engine/Source/Programs/Unsync/Private/UnsyncError.h
@@ -156,6 +156,8 @@ struct FAtomicError
 	std::atomic_flag	  Flag;
 	std::optional<FError> Data;
 
+	operator bool() const { return Test(); }
+
 	bool Test() const { return Flag.test(); }
 	bool Set(FError&& InData)
 	{
