unsync - Add EFileMode::IgnoreDryRun mode to always allow certain file write operations

* Change EFileMode from regular enum to flags
* Add GetAnonymizedMachineId helper function
* Use MachineGuid registry key on Windows

diff --git a/Engine/Source/Programs/Unsync/Private/UnsyncAuth.cpp b/Engine/Source/Programs/Unsync/Private/UnsyncAuth.cpp
--- a/Engine/Source/Programs/Unsync/Private/UnsyncAuth.cpp
+++ b/Engine/Source/Programs/Unsync/Private/UnsyncAuth.cpp
@@ -623,7 +623,11 @@ static FAuthTokenCache GAuthTokenCache;
 bool
 SaveAuthToken(const FPath& Path, const FAuthToken& AuthToken)
 {
-	bool bWrittenOk = WriteBufferToFile(Path, (const uint8*)AuthToken.Raw.data(), AuthToken.Raw.length());
+	bool bWrittenOk = WriteBufferToFile(Path,
+										(const uint8*)AuthToken.Raw.data(),
+										AuthToken.Raw.length(),
+										EFileMode::CreateWriteOnly | EFileMode::IgnoreDryRun);
+
 	if (!bWrittenOk)
 	{
 		return false;
@@ -801,16 +805,8 @@ Authenticate(const FRemoteDesc& RemoteDesc, int32 RefreshThreshold)
 	{
 		CreateDirectories(TokenCachePath->parent_path());
 
-		// Allow saving tokens during dry run
-		// TODO: need a dedicated file flag to allow writes during dry run
-
-		const bool bPrevDryRun = GDryRun;
-		GDryRun				   = false;
-
 		bool bSaved = SaveAuthToken(*TokenCachePath, FreshTokenResult.GetData());
 
-		GDryRun = bPrevDryRun;
-
 		if (bSaved)
 		{
 			UNSYNC_VERBOSE2(L"Saved authentication token to file: %ls", TokenCachePath->wstring().c_str());
diff --git a/Engine/Source/Programs/Unsync/Private/UnsyncCommon.h b/Engine/Source/Programs/Unsync/Private/UnsyncCommon.h
--- a/Engine/Source/Programs/Unsync/Private/UnsyncCommon.h
+++ b/Engine/Source/Programs/Unsync/Private/UnsyncCommon.h
@@ -72,4 +72,44 @@ ClobberT(T& Object)
 #	define UNSYNC_THIRD_PARTY_INCLUDES_END _Pragma("GCC diagnostic pop")
 #endif
 
+#define UNSYNC_ENUM_CLASS_FLAGS(T, S)                      \
+	inline constexpr T operator^(T a, T b)                 \
+	{                                                      \
+		return (T)(static_cast<S>(a) ^ static_cast<S>(b)); \
+	}                                                      \
+	inline constexpr T operator&(T a, T b)                 \
+	{                                                      \
+		return (T)(static_cast<S>(a) & static_cast<S>(b)); \
+	}                                                      \
+	inline constexpr T operator|(T a, T b)                 \
+	{                                                      \
+		return (T)(static_cast<S>(a) | static_cast<S>(b)); \
+	}                                                      \
+	inline constexpr bool operator!(T a)                   \
+	{                                                      \
+		return !static_cast<S>(a);                         \
+	}                                                      \
+	inline constexpr bool operator==(T a, S b)             \
+	{                                                      \
+		return static_cast<S>(a) == b;                     \
+	}                                                      \
+	inline constexpr bool operator!=(T a, S b)             \
+	{                                                      \
+		return static_cast<S>(a) != b;                     \
+	}
+
+template<typename Enum>
+constexpr bool
+EnumHasAllFlags(Enum Flags, Enum Contains)
+{
+	return (Flags & Contains) == Contains;
+}
+
+template<typename Enum>
+constexpr bool
+EnumHasAnyFlags(Enum Flags, Enum Contains)
+{
+	return (Flags & Contains) != 0;
+}
+
 }  // namespace unsync
diff --git a/Engine/Source/Programs/Unsync/Private/UnsyncCore.cpp b/Engine/Source/Programs/Unsync/Private/UnsyncCore.cpp
--- a/Engine/Source/Programs/Unsync/Private/UnsyncCore.cpp
+++ b/Engine/Source/Programs/Unsync/Private/UnsyncCore.cpp
@@ -2951,7 +2951,7 @@ struct FPooledProxy
 };
 
 static bool
-DownloadFileIfNewer(const FRemoteDesc& RemoteDesc, const FPath& Source, const FPath& Target)
+DownloadFileIfNewer(const FRemoteDesc& RemoteDesc, const FPath& Source, const FPath& Target, EFileMode TargetFileMode)
 {
 	using FDirectoryListing		 = ProxyQuery::FDirectoryListing;
 	using FDirectoryListingEntry = ProxyQuery::FDirectoryListingEntry;
@@ -3008,7 +3008,7 @@ DownloadFileIfNewer(const FRemoteDesc& RemoteDesc, const FPath& Source, const FP
 			return false;
 		}
 
-		bool bFileWritten = WriteBufferToFile(Target, FileBuffer);
+		bool bFileWritten = WriteBufferToFile(Target, FileBuffer, TargetFileMode);
 		if (!bFileWritten)
 		{
 			UNSYNC_ERROR(L"Failed to write downloaded file '%ls'", Target.wstring().c_str());
@@ -3068,14 +3068,11 @@ LoadAndMergeSourceManifest(FDirectoryManifest& Output,
 
 	if (bDownloadManifestFromProxy)
 	{
-		// TODO: add a mechanism to suppress dry run in a scope in a thread-safe way
-		const bool bPrevDryRun = GDryRun;
-		GDryRun				   = false;
-
 		UNSYNC_LOG_INDENT;
-		bool bDownloadedOk = DownloadFileIfNewer(ProxyPool.RemoteDesc, SourceManifestPath, SourceManifestTempPath);
-
-		GDryRun = bPrevDryRun;
+		bool bDownloadedOk = DownloadFileIfNewer(ProxyPool.RemoteDesc,
+												 SourceManifestPath,
+												 SourceManifestTempPath,
+												 EFileMode::CreateWriteOnly | EFileMode::IgnoreDryRun);
 
 		if (!bDownloadedOk)
 		{
@@ -3119,20 +3116,6 @@ LoadAndMergeSourceManifest(FDirectoryManifest& Output,
 	return MergeManifests(Output, LoadedManifest, bCaseSensitiveTargetFileSystem);
 }
 
-static std::string
-GetAnonymizedHostName()
-{
-	std::string HostName = GetCurrentHostName();
-	if (HostName.empty())
-	{
-		return {};
-	}
-	HostName += " {22FF4421-8CAC-4A14-9E4C-780AAF8BBF2A}";
-	FHash128	MachineId = HashBlake3String<FHash128>(HostName);
-	std::string Result	  = HashToHexString(MachineId);
-	return Result;
-}
-
 struct FFileSyncTaskBatch
 {
 	std::vector<const FFileSyncTask*> FileTasks;
@@ -4104,7 +4087,7 @@ SyncDirectory(const FSyncDirectoryOptions& SyncOptions)
 		Event.ClientVersion		 = GetVersionString();
 		Event.Session			 = ProxyPool.GetSessionId();
 		Event.Source			 = ConvertWideToUtf8(SourcePath.wstring());
-		Event.ClientHostNameHash = GetAnonymizedHostName();
+		Event.ClientHostNameHash = GetAnonymizedMachineIdString();
 		Event.TotalBytes		 = TotalSourceSize;
 		Event.SourceBytes		 = StatSourceBytes;
 		Event.BaseBytes			 = StatBaseBytes;
diff --git a/Engine/Source/Programs/Unsync/Private/UnsyncFile.cpp b/Engine/Source/Programs/Unsync/Private/UnsyncFile.cpp
--- a/Engine/Source/Programs/Unsync/Private/UnsyncFile.cpp
+++ b/Engine/Source/Programs/Unsync/Private/UnsyncFile.cpp
@@ -95,7 +95,7 @@ struct FCreateFileInfo
 
 	FCreateFileInfo(EFileMode Mode)
 	{
-		switch (Mode)
+		switch (Mode & EFileMode::CommonModeMask)
 		{
 			default:
 			case EFileMode::ReadOnly:
@@ -108,7 +108,7 @@ struct FCreateFileInfo
 				break;
 			case EFileMode::CreateReadWrite:
 			case EFileMode::CreateWriteOnly:
-				UNSYNC_ASSERT(!GDryRun);
+				UNSYNC_ASSERT(!GDryRun || EnumHasAnyFlags(Mode, EFileMode::IgnoreDryRun));
 				FileAccess	= GENERIC_READ | GENERIC_WRITE;
 				Share		= FILE_SHARE_WRITE;
 				Disposition = CREATE_ALWAYS;
@@ -194,7 +194,7 @@ FWindowsFile::OpenFileHandle(EFileMode InMode)
 {
 	FCreateFileInfo Info(InMode);
 	Info.FileFlags |= FILE_FLAG_OVERLAPPED;
-	if (InMode == EFileMode::ReadOnlyUnbuffered && !GForceBufferedFiles)
+	if (EnumHasAnyFlags(InMode, EFileMode::Unbuffered) && !GForceBufferedFiles)
 	{
 		Info.FileFlags |= FILE_FLAG_NO_BUFFERING;
 	}
@@ -397,7 +397,7 @@ FWindowsFile::Write(const void* Data, uint64 DestOffset, uint64 TotalSize)
 uint64
 FWindowsFile::Read(void* Dest, uint64 SourceOffset, uint64 ReadSize)
 {
-	UNSYNC_ASSERTF(Mode != EFileMode::ReadOnlyUnbuffered, L"ReadUnbuffered mode is not supported for non-async reads");
+	UNSYNC_ASSERTF((Mode & EFileMode::Unbuffered) == 0, L"Unbuffered files only support ReadAsync");
 	UNSYNC_ASSERT(IsReadable(Mode));
 
 	LARGE_INTEGER Pos;
@@ -465,7 +465,7 @@ FWindowsFile::ReadAsync(uint64 SourceOffset, uint64 Size, uint64 UserData, IOCal
 		CompleteReadCommand(Cmd);
 	}
 
-	if (Mode == EFileMode::ReadOnlyUnbuffered)
+	if (EnumHasAnyFlags(Mode, EFileMode::Unbuffered))
 	{
 		uint64 OriginalSize	 = Size;
 		uint64 OriginalBegin = SourceOffset;
@@ -862,14 +862,27 @@ ReadFileToBuffer(const FPath& Filename)
 }
 
 bool
-WriteBufferToFile(const FPath& Filename, const uint8* Data, uint64 Size)
+WriteBufferToFile(const FPath& Filename, const uint8* Data, uint64 Size, EFileMode FileMode)
 {
 	UNSYNC_LOG_INDENT;
-	UNSYNC_ASSERT(Data);
-	UNSYNC_ASSERT(Size);
-	UNSYNC_ASSERT(!GDryRun);
 
-	FNativeFile File(Filename, EFileMode::CreateReadWrite, Size);
+	if (Data == nullptr)
+	{
+		UNSYNC_ERROR(L"WriteBufferToFile called with null buffer");
+		return false;
+	}
+	if (Size == 0)
+	{
+		UNSYNC_ERROR(L"WriteBufferToFile called with zero size buffer");
+		return false;
+	}
+	if (GDryRun && !EnumHasAnyFlags(FileMode, EFileMode::IgnoreDryRun))
+	{
+		UNSYNC_ERROR(L"WriteBufferToFile called in dry run mode");
+		return false;
+	}
+
+	FNativeFile File(Filename, FileMode, Size);
 
 	if (File.IsValid())
 	{
@@ -886,9 +899,15 @@ WriteBufferToFile(const FPath& Filename, const uint8* Data, uint64 Size)
 }
 
 bool
-WriteBufferToFile(const FPath& Filename, const FBuffer& Buffer)
+WriteBufferToFile(const FPath& Filename, const FBuffer& Buffer, EFileMode FileMode)
+{
+	return WriteBufferToFile(Filename, Buffer.Data(), Buffer.Size(), FileMode);
+}
+
+bool
+WriteBufferToFile(const FPath& Filename, const std::string& Buffer, EFileMode FileMode)
 {
-	return WriteBufferToFile(Filename, Buffer.Data(), Buffer.Size());
+	return WriteBufferToFile(Filename, (const uint8*)Buffer.data(), Buffer.length(), FileMode);
 }
 
 struct FIOBufferCache
diff --git a/Engine/Source/Programs/Unsync/Private/UnsyncFile.h b/Engine/Source/Programs/Unsync/Private/UnsyncFile.h
--- a/Engine/Source/Programs/Unsync/Private/UnsyncFile.h
+++ b/Engine/Source/Programs/Unsync/Private/UnsyncFile.h
@@ -14,18 +14,33 @@ namespace unsync {
 
 extern bool GForceBufferedFiles;
 
-enum class EFileMode
-{
-	ReadOnly,
-	ReadOnlyUnbuffered,
-	CreateReadWrite,
-	CreateWriteOnly,
+enum class EFileMode : uint32 {
+	None	   = 0,
+
+	Read	   = 1 << 0,
+	Write	   = 1 << 1,
+	Create	   = 1 << 2,
+	Unbuffered = 1 << 3,
+
+	// Extended modes
+	IgnoreDryRun = 1 << 4,	// allow write operations even in dry run mode
+
+	// Commonly used mode combinations
+	ReadOnly		   = Read,
+	ReadOnlyUnbuffered = Read | Unbuffered,
+	CreateReadWrite	   = Read | Write | Create,
+	CreateWriteOnly	   = Write | Create,
+
+	// Masks
+	CommonModeMask	 = Create | Read | Write | Unbuffered,
+	ExtendedModeMask = ~CommonModeMask,
 };
+UNSYNC_ENUM_CLASS_FLAGS(EFileMode, uint32)
 
 inline bool
 IsReadOnly(EFileMode Mode)
 {
-	switch (Mode)
+	switch (Mode & EFileMode::CommonModeMask)
 	{
 		case EFileMode::ReadOnly:
 		case EFileMode::ReadOnlyUnbuffered:
@@ -38,32 +53,19 @@ IsReadOnly(EFileMode Mode)
 inline bool
 IsWriteOnly(EFileMode Mode)
 {
-	switch (Mode)
-	{
-		case EFileMode::CreateWriteOnly:
-			return true;
-		default:
-			return false;
-	}
+	return (Mode & EFileMode::Read) == 0;
 }
 
 inline bool
 IsReadable(EFileMode Mode)
 {
-	return !IsWriteOnly(Mode);
+	return (Mode & EFileMode::Read) != 0;
 }
 
 inline bool
 IsWritable(EFileMode Mode)
 {
-	switch (Mode)
-	{
-		case EFileMode::CreateReadWrite:
-		case EFileMode::CreateWriteOnly:
-			return true;
-		default:
-			return false;
-	}
+	return (Mode & EFileMode::Write) != 0;
 }
 
 struct FIOBuffer
@@ -422,8 +424,9 @@ struct FIOReaderStream
 };
 
 FBuffer ReadFileToBuffer(const FPath& Filename);
-bool	WriteBufferToFile(const FPath& Filename, const uint8* Data, uint64 Size);
-bool	WriteBufferToFile(const FPath& Filename, const FBuffer& Buffer);
+bool	WriteBufferToFile(const FPath& Filename, const uint8* Data, uint64 Size, EFileMode FileMode = EFileMode::CreateWriteOnly);
+bool	WriteBufferToFile(const FPath& Filename, const FBuffer& Buffer, EFileMode FileMode = EFileMode::CreateWriteOnly);
+bool	WriteBufferToFile(const FPath& Filename, const std::string& Buffer, EFileMode FileMode = EFileMode::CreateWriteOnly);
 
 struct FFileAttributes
 {
diff --git a/Engine/Source/Programs/Unsync/Private/UnsyncUtil.cpp b/Engine/Source/Programs/Unsync/Private/UnsyncUtil.cpp
--- a/Engine/Source/Programs/Unsync/Private/UnsyncUtil.cpp
+++ b/Engine/Source/Programs/Unsync/Private/UnsyncUtil.cpp
@@ -2,6 +2,7 @@
 
 #include "UnsyncUtil.h"
 #include "UnsyncFile.h"
+#include "UnsyncSocket.h"
 
 #if UNSYNC_PLATFORM_WINDOWS
 #	include <Windows.h>
@@ -14,6 +15,7 @@
 #	pragma comment(lib, "Crypt32.lib")
 #	pragma comment(lib, "Bcrypt.lib")
 #	pragma comment(lib, "Mpr.lib")
+#	pragma comment(lib, "Advapi32.lib") // for registry access
 #endif	// UNSYNC_PLATFORM_WINDOWS
 
 #include <stdlib.h>
@@ -498,4 +500,48 @@ FormatSystemErrorMessage(int32 ErrorCode)
 	return fmt::format("Error code {}: {}", ErrorCode, ErrorMessage);
 }
 
+FHash256
+GetAnonymizedMachineId(std::string_view Salt)
+{
+	std::string Seed;
+
+	Seed += Salt;
+	Seed += GetCurrentHostName();
+	Seed += " {22FF4421-8CAC-4A14-9E4C-780AAF8BBF2A}";
+
+#if UNSYNC_PLATFORM_WINDOWS
+	{
+		HKEY Key = {};
+		if (RegOpenKeyA(HKEY_LOCAL_MACHINE, "SOFTWARE\\Microsoft\\Cryptography", &Key) == ERROR_SUCCESS)
+		{
+			char  Buffer[512] = {};
+			DWORD BufferSize  = sizeof(Buffer);
+			auto  Status	  = RegQueryValueExA(Key, "MachineGuid", nullptr, nullptr, (LPBYTE)Buffer, &BufferSize);
+			if (Status == ERROR_SUCCESS && BufferSize > 1)
+			{
+				std::string_view MachineGuid = std::string_view(Buffer, BufferSize - 1);
+				Seed += " MachineGuid ";
+				Seed += MachineGuid;
+			}
+			RegCloseKey(Key);
+		}
+	}
+#endif	// UNSYNC_PLATFORM_WINDOWS
+
+	// TODO: read `/etc/machine-id` on linux
+	// TODO: use `ioreg -rd1 -c IOPlatformExpertDevice` to get IOPlatformUUID on mac
+
+	FHash256 Result = HashBlake3String<FHash256>(Seed);
+
+	return Result;
+}
+
+std::string
+GetAnonymizedMachineIdString(std::string_view Seed)
+{
+	FHash256	MachineId = GetAnonymizedMachineId(Seed);
+	std::string Result	  = HashToHexString(MachineId);
+	return Result;
+}
+
 }  // namespace unsync
diff --git a/Engine/Source/Programs/Unsync/Private/UnsyncUtil.h b/Engine/Source/Programs/Unsync/Private/UnsyncUtil.h
--- a/Engine/Source/Programs/Unsync/Private/UnsyncUtil.h
+++ b/Engine/Source/Programs/Unsync/Private/UnsyncUtil.h
@@ -330,6 +330,9 @@ void OpenUrlInDefaultBrowser(const char* Address);
 
 FPath GetUserHomeDirectory();
 
+FHash256	GetAnonymizedMachineId(std::string_view Salt = {});
+std::string GetAnonymizedMachineIdString(std::string_view Salt = {});
+
 // Returns string in format 'Error code 123: Some description.`
 std::string FormatSystemErrorMessage(int32 ErrorCode);
 
