unsync - Add search query to discover available builds through proxy

diff --git a/Engine/Source/Programs/Unsync/Private/UnsyncCmdQuery.cpp b/Engine/Source/Programs/Unsync/Private/UnsyncCmdQuery.cpp
--- a/Engine/Source/Programs/Unsync/Private/UnsyncCmdQuery.cpp
+++ b/Engine/Source/Programs/Unsync/Private/UnsyncCmdQuery.cpp
@@ -7,9 +7,11 @@
 #include "UnsyncThread.h"
 #include "UnsyncUtil.h"
 #include "UnsyncAuth.h"
+#include "UnsyncPool.h"
 
 #include <float.h>
 #include <algorithm>
+#include <regex>
 #include <json11.hpp>
 
 #include <fmt/format.h>
@@ -149,9 +151,15 @@ CmdQueryList(const FCmdQueryOptions& Options)
 		return -1;
 	}
 
+	if (Options.Args.empty())
+	{
+		UNSYNC_ERROR(L"Path argument is required");
+		return -1;
+	}
+
 	FHttpConnection Connection = FHttpConnection::CreateDefaultHttps(Options.Remote);
 
-	std::string Url = fmt::format("/api/v1/list?{}", Options.Args);
+	std::string Url = fmt::format("/api/v1/list?{}", Options.Args[0]);
 
 	FHttpRequest Request;
 	Request.Url			   = Url;
@@ -175,10 +183,176 @@ CmdQueryList(const FCmdQueryOptions& Options)
 	return 0;
 }
 
+int32
+CmdQuerySearch(const FCmdQueryOptions& Options)
+{
+	TResult<FAuthToken> AuthToken = Authenticate(Options.Remote, 5 * 60);
+	if (!AuthToken.IsOk())
+	{
+		LogError(AuthToken.GetError());
+		return -1;
+	}
+
+	if (Options.Args.empty())
+	{
+		UNSYNC_ERROR(L"Path argument is required");
+		return -1;
+	}
+
+	auto CreateConnection = [Remote = Options.Remote]
+	{
+		FTlsClientSettings TlsSettings;
+		TlsSettings.Subject			   = Remote.HostAddress.data();
+		TlsSettings.bVerifyCertificate = Remote.bTlsVerifyCertificate;
+		return new FHttpConnection(Remote.HostAddress, Remote.HostPort, &TlsSettings);
+	};
+
+	TObjectPool<FHttpConnection> ConnectionPool(CreateConnection);
+
+	const std::string& RootPath = Options.Args[0];
+	std::vector<std::regex> SubdirPatterns;
+
+	UNSYNC_VERBOSE("Searching '%hs'", RootPath.c_str());
+	if (Options.Args.size() > 1)
+	{
+		UNSYNC_LOG_INDENT;
+		UNSYNC_VERBOSE("Subdirectory patterns:", RootPath.c_str());
+		UNSYNC_LOG_INDENT;
+		for (size_t i = 1; i < Options.Args.size(); ++i)
+		{
+			UNSYNC_VERBOSE("%hs", Options.Args[i].c_str());
+			std::regex Pattern = std::regex(Options.Args[i], std::regex_constants::icase);
+			SubdirPatterns.push_back(Pattern);
+		}
+	}
+
+	struct FEntry
+	{
+		std::string Path;
+		uint32		Depth = 0;
+	};
+
+	std::vector<FEntry> PendingDirectories;
+	{
+		FEntry RootEntry;
+		RootEntry.Path	= RootPath;
+		RootEntry.Depth = 0;
+		PendingDirectories.push_back(RootEntry);
+	}
+
+	struct FTaskContext
+	{
+		std::mutex Mutex;
+		FSemaphore ConnectionSemaphore = FSemaphore(8);
+		std::vector<std::string> FoundDirectories;
+		bool					 bParentThreadVerbose = false;
+		int32					 ParentThreadIndent = 0;
+	};
+
+	FTaskGroup Tasks;
+	FTaskContext Context;
+
+	Context.bParentThreadVerbose = GLogVerbose;
+	Context.ParentThreadIndent	 = GLogIndent;
+
+	std::function<void(std::string, int32)> ExploreDirectory =
+		[&Context, &AuthToken, &ConnectionPool, &ExploreDirectory, &SubdirPatterns, &Tasks](std::string Path, int32 CurrentDepth)
+	{
+		FLogVerbosityScope VerboseScope(Context.bParentThreadVerbose);
+		FLogIndentScope	   IndentScope(Context.ParentThreadIndent, true);
+
+		UNSYNC_VERBOSE2(L"Listing '%hs'", Path.c_str());
+
+		std::string Url = fmt::format("/api/v1/list?{}", Path);
+
+		FHttpRequest Request;
+		Request.Url			= Url;
+		Request.Method		= EHttpMethod::GET;
+		Request.BearerToken = AuthToken->Access;
+
+		Context.ConnectionSemaphore.Acquire();
+
+		std::unique_ptr<FHttpConnection> Connection = ConnectionPool.Acquire();
+		FHttpResponse					 Response	= HttpRequest(*Connection, Request);
+		ConnectionPool.Release(std::move(Connection));
+
+		Context.ConnectionSemaphore.Release();
+
+		if (!Response.Success())
+		{
+			// TODO: report warning
+			return;
+		}
+
+		Response.Buffer.PushBack(0);
+
+		using namespace ProxyQuery;
+		TResult<FDirectoryListing> DirectoryListingResult = FDirectoryListing::FromJson((const char*)Response.Buffer.Data());
+		if (DirectoryListingResult.IsError())
+		{
+			// TODO: report warning
+			return;
+		}
+
+		const FDirectoryListing& DirectoryListing = DirectoryListingResult.GetData();
+
+		for (const FDirectoryListingEntry& Subdir : DirectoryListing.Entries)
+		{
+			const std::regex& RequiredPattern = SubdirPatterns[CurrentDepth];
+
+			if (std::regex_match(Subdir.Name, RequiredPattern, std::regex_constants::match_any))
+			{
+				UNSYNC_VERBOSE2(L"Matched: '%hs'", Subdir.Name.c_str());
+
+				FEntry NextEntry;
+				NextEntry.Path	= Path + "\\" + Subdir.Name;
+				NextEntry.Depth = CurrentDepth + 1;
+
+				if (NextEntry.Depth == SubdirPatterns.size())
+				{
+					{
+						std::lock_guard<std::mutex> LockGuard(Context.Mutex);
+						Context.FoundDirectories.push_back(NextEntry.Path);
+					}
+					continue;
+				}
+
+				Tasks.run([ExploreDirectory, NextEntry]() {
+					ExploreDirectory(NextEntry.Path, NextEntry.Depth);
+				});
+			}
+		}
+	};
+
+	ExploreDirectory(RootPath, 0);
+	Tasks.wait();
+
+	std::vector<std::string> FoundDirectories = std::move(Context.FoundDirectories);
+
+	std::sort(FoundDirectories.begin(), FoundDirectories.end());
+
+	LogPrintf(ELogLevel::MachineReadable, L"[\n");
+	for (size_t i = 0; i < FoundDirectories.size(); ++i)
+	{
+		const std::string& ResultEntry	 = FoundDirectories[i];
+		const char*		   TrailingComma = i + 1 == FoundDirectories.size() ? "" : ",";
+		LogPrintf(ELogLevel::MachineReadable, L"  {\"path\":\"%hs\"}%hs\n", ResultEntry.c_str(), TrailingComma);
+	}
+	LogPrintf(ELogLevel::MachineReadable, L"]\n");
+
+	return 0;
+}
+
 int32
 CmdQueryFile(const FCmdQueryOptions& Options)
 {
-	UNSYNC_LOG(L"Downloading file: '%hs'", Options.Args.c_str());
+	if (Options.Args.empty())
+	{
+		UNSYNC_ERROR(L"Path argument is required");
+		return -1;
+	}
+
+	UNSYNC_LOG(L"Downloading file: '%hs'", Options.Args[0].c_str());
 	UNSYNC_LOG_INDENT;
 
 	std::unique_ptr<FNativeFile> ResultWriter;
@@ -189,7 +363,7 @@ CmdQueryFile(const FCmdQueryOptions& Options)
 		return *ResultWriter;
 	};
 
-	TResult<> Response = ProxyQuery::DownloadFile(Options.Remote, Options.Args, OutputCallback);
+	TResult<> Response = ProxyQuery::DownloadFile(Options.Remote, Options.Args[0], OutputCallback);
 
 	if (Response.IsOk())
 	{
@@ -214,6 +388,10 @@ CmdQuery(const FCmdQueryOptions& Options)
 	{
 		return CmdQueryList(Options);
 	}
+	else if (Options.Query == "search")
+	{
+		return CmdQuerySearch(Options);
+	}
 	else if (Options.Query == "file")
 	{
 		return CmdQueryFile(Options);
diff --git a/Engine/Source/Programs/Unsync/Private/UnsyncCmdQuery.h b/Engine/Source/Programs/Unsync/Private/UnsyncCmdQuery.h
--- a/Engine/Source/Programs/Unsync/Private/UnsyncCmdQuery.h
+++ b/Engine/Source/Programs/Unsync/Private/UnsyncCmdQuery.h
@@ -9,10 +9,10 @@ namespace unsync {
 
 struct FCmdQueryOptions
 {
-	std::string Query;
-	std::string Args;
-	FPath		OutputPath;
-	FRemoteDesc Remote;
+	std::string				 Query;
+	std::vector<std::string> Args;
+	FPath					 OutputPath;
+	FRemoteDesc				 Remote;
 };
 
 int32 CmdQuery(const FCmdQueryOptions& Options);
diff --git a/Engine/Source/Programs/Unsync/Private/UnsyncMain.cpp b/Engine/Source/Programs/Unsync/Private/UnsyncMain.cpp
--- a/Engine/Source/Programs/Unsync/Private/UnsyncMain.cpp
+++ b/Engine/Source/Programs/Unsync/Private/UnsyncMain.cpp
@@ -73,7 +73,7 @@ InnerMain(int Argc, char** Argv)
 	std::string				 ProtocolName = "jupiter";
 	std::string				 HttpHeaderFilenameUtf8;
 	std::string				 QueryStringUtf8;
-	std::string				 QueryArgsUtf8;
+	std::vector<std::string> QueryArgsUtf8;
 	std::string				 ScavengeRootUtf8;
 	bool					 bForceOperation	 = false;
 	bool					 bAllowInsecureTls	 = false;
diff --git a/Engine/Source/Programs/Unsync/Private/UnsyncProxy.cpp b/Engine/Source/Programs/Unsync/Private/UnsyncProxy.cpp
--- a/Engine/Source/Programs/Unsync/Private/UnsyncProxy.cpp
+++ b/Engine/Source/Programs/Unsync/Private/UnsyncProxy.cpp
@@ -510,31 +510,12 @@ ProxyQuery::Hello(const FRemoteDesc& RemoteDesc, bool bAnonymous)
 }
 
 TResult<ProxyQuery::FDirectoryListing>
-ProxyQuery::ListDirectory(const FRemoteDesc& Remote, const std::string& Path)
+ProxyQuery::FDirectoryListing::FromJson(const char* JsonString)
 {
 	FDirectoryListing Result;
 
-	TResult<FAuthToken> AuthToken = Authenticate(Remote, 5 * 60);
-	if (!AuthToken.IsOk())
-	{
-		return MoveError<FDirectoryListing>(AuthToken);
-	}
-
-	FHttpConnection Connection = FHttpConnection::CreateDefaultHttps(Remote);
-
-	std::string Url = fmt::format("/api/v1/list?{}", Path);
-
-	FHttpRequest Request;
-	Request.Url			= Url;
-	Request.Method		= EHttpMethod::GET;
-	Request.BearerToken = AuthToken->Access;
-
-	FHttpResponse Response = HttpRequest(Connection, Request);
-
-	Response.Buffer.PushBack(0);
-
 	std::string	 JsonErrorString;
-	json11::Json JsonObject = json11::Json::parse((const char*)Response.Buffer.Data(), JsonErrorString);
+	json11::Json JsonObject = json11::Json::parse(JsonString, JsonErrorString);
 
 	if (!JsonErrorString.empty())
 	{
@@ -574,6 +555,31 @@ ProxyQuery::ListDirectory(const FRemoteDesc& Remote, const std::string& Path)
 	return ResultOk(std::move(Result));
 }
 
+TResult<ProxyQuery::FDirectoryListing>
+ProxyQuery::ListDirectory(const FRemoteDesc& Remote, const std::string& Path)
+{
+	TResult<FAuthToken> AuthToken = Authenticate(Remote, 5 * 60);
+	if (!AuthToken.IsOk())
+	{
+		return MoveError<FDirectoryListing>(AuthToken);
+	}
+
+	FHttpConnection Connection = FHttpConnection::CreateDefaultHttps(Remote);
+
+	std::string Url = fmt::format("/api/v1/list?{}", Path);
+
+	FHttpRequest Request;
+	Request.Url			= Url;
+	Request.Method		= EHttpMethod::GET;
+	Request.BearerToken = AuthToken->Access;
+
+	FHttpResponse Response = HttpRequest(Connection, Request);
+
+	Response.Buffer.PushBack(0);
+
+	return FDirectoryListing::FromJson((const char*)Response.Buffer.Data());
+}
+
 TResult<>
 ProxyQuery::DownloadFile(const FRemoteDesc& Remote, const std::string& Path, ProxyQuery::FDownloadOutputCallback OutputCallback)
 {
diff --git a/Engine/Source/Programs/Unsync/Private/UnsyncProxy.h b/Engine/Source/Programs/Unsync/Private/UnsyncProxy.h
--- a/Engine/Source/Programs/Unsync/Private/UnsyncProxy.h
+++ b/Engine/Source/Programs/Unsync/Private/UnsyncProxy.h
@@ -217,6 +217,8 @@ struct FDirectoryListingEntry
 struct FDirectoryListing
 {
 	std::vector<FDirectoryListingEntry> Entries;
+
+	static TResult<FDirectoryListing> FromJson(const char* JsonString);
 };
 
 TResult<FDirectoryListing> ListDirectory(const FRemoteDesc& Remote, const std::string& Path);
