unsync - Output search query results as simple string array json

* Escape strings when formatting adhoc json output
* Bump version to 1.0.57

diff --git a/Engine/Source/Programs/Unsync/Private/UnsyncCmdQuery.cpp b/Engine/Source/Programs/Unsync/Private/UnsyncCmdQuery.cpp
--- a/Engine/Source/Programs/Unsync/Private/UnsyncCmdQuery.cpp
+++ b/Engine/Source/Programs/Unsync/Private/UnsyncCmdQuery.cpp
@@ -128,11 +128,11 @@ CmdQueryMirrors(const FCmdQueryOptions& Options)
 
 		LogPrintf(ELogLevel::MachineReadable,
 				  L"  {\"address\":\"%hs\", \"port\":%d, \"ok\":%hs, \"ping\":%d, \"name\":\"%hs\"}%hs\n",
-				  Mirror.Address.c_str(),
+				  StringEscape(Mirror.Address).c_str(),
 				  Mirror.Port,
 				  Mirror.Ping > 0 ? "true" : "false",
 				  int32(Mirror.Ping * 1000.0),
-				  Mirror.Name.c_str(),
+				  StringEscape(Mirror.Name).c_str(),
 				  I + 1 == Mirrors.size() ? "" : ",");
 	}
 
@@ -336,7 +336,7 @@ CmdQuerySearch(const FCmdQueryOptions& Options)
 	{
 		const std::string& ResultEntry	 = FoundDirectories[i];
 		const char*		   TrailingComma = i + 1 == FoundDirectories.size() ? "" : ",";
-		LogPrintf(ELogLevel::MachineReadable, L"  {\"path\":\"%hs\"}%hs\n", ResultEntry.c_str(), TrailingComma);
+		LogPrintf(ELogLevel::MachineReadable, L"  \"%hs\"%hs\n", StringEscape(ResultEntry).c_str(), TrailingComma);
 	}
 	LogPrintf(ELogLevel::MachineReadable, L"]\n");
 
diff --git a/Engine/Source/Programs/Unsync/Private/UnsyncCore.cpp b/Engine/Source/Programs/Unsync/Private/UnsyncCore.cpp
--- a/Engine/Source/Programs/Unsync/Private/UnsyncCore.cpp
+++ b/Engine/Source/Programs/Unsync/Private/UnsyncCore.cpp
@@ -27,7 +27,7 @@ UNSYNC_THIRD_PARTY_INCLUDES_START
 #include <md5-sse2.h>
 UNSYNC_THIRD_PARTY_INCLUDES_END
 
-#define UNSYNC_VERSION_STR "1.0.56"
+#define UNSYNC_VERSION_STR "1.0.57"
 
 namespace unsync {
 
diff --git a/Engine/Source/Programs/Unsync/Private/UnsyncUtil.cpp b/Engine/Source/Programs/Unsync/Private/UnsyncUtil.cpp
--- a/Engine/Source/Programs/Unsync/Private/UnsyncUtil.cpp
+++ b/Engine/Source/Programs/Unsync/Private/UnsyncUtil.cpp
@@ -186,6 +186,71 @@ StringToUpper(const std::wstring& Input)
 	return Result;
 }
 
+std::string
+StringEscape(const std::string_view Input)
+{
+	// Adapted from Json11
+
+	std::string Result;
+
+	for (size_t i = 0; i < Input.length(); i++)
+	{
+		const char C = Input[i];
+		if (C == '\\')
+		{
+			Result += "\\\\";
+		}
+		else if (C == '"')
+		{
+			Result += "\\\"";
+		}
+		else if (C == '\b')
+		{
+			Result += "\\b";
+		}
+		else if (C == '\f')
+		{
+			Result += "\\f";
+		}
+		else if (C == '\n')
+		{
+			Result += "\\n";
+		}
+		else if (C == '\r')
+		{
+			Result += "\\r";
+		}
+		else if (C == '\t')
+		{
+			Result += "\\t";
+		}
+		else if (static_cast<uint8_t>(C) <= 0x1f)
+		{
+			char buf[8];
+			snprintf(buf, sizeof buf, "\\u%04x", C);
+			Result += buf;
+		}
+		else if (static_cast<uint8_t>(C) == 0xe2 && static_cast<uint8_t>(Input[i + 1]) == 0x80 &&
+				 static_cast<uint8_t>(Input[i + 2]) == 0xa8)
+		{
+			Result += "\\u2028";
+			i += 2;
+		}
+		else if (static_cast<uint8_t>(C) == 0xe2 && static_cast<uint8_t>(Input[i + 1]) == 0x80 &&
+				 static_cast<uint8_t>(Input[i + 2]) == 0xa9)
+		{
+			Result += "\\u2029";
+			i += 2;
+		}
+		else
+		{
+			Result += C;
+		}
+	}
+
+	return Result;
+}
+
 FDfsMirrorInfo
 DfsEnumerate(const FPath& Root)
 {
diff --git a/Engine/Source/Programs/Unsync/Private/UnsyncUtil.h b/Engine/Source/Programs/Unsync/Private/UnsyncUtil.h
--- a/Engine/Source/Programs/Unsync/Private/UnsyncUtil.h
+++ b/Engine/Source/Programs/Unsync/Private/UnsyncUtil.h
@@ -138,6 +138,8 @@ std::string	 ConvertWideToUtf8(std::wstring_view StringWide);
 std::wstring StringToLower(const std::wstring& Input);
 std::wstring StringToUpper(const std::wstring& Input);
 
+std::string StringEscape(const std::string_view Input);
+
 // Returns a list of alternative DFS paths for a given root
 struct FDfsStorageInfo
 {
