unsync - Extend output of the search query to include some file metadata and use paths relative to root of the query

* Bump version to 1.0.59

diff --git a/Engine/Source/Programs/Unsync/Private/UnsyncCmdQuery.cpp b/Engine/Source/Programs/Unsync/Private/UnsyncCmdQuery.cpp
--- a/Engine/Source/Programs/Unsync/Private/UnsyncCmdQuery.cpp
+++ b/Engine/Source/Programs/Unsync/Private/UnsyncCmdQuery.cpp
@@ -186,6 +186,8 @@ CmdQueryList(const FCmdQueryOptions& Options)
 int32
 CmdQuerySearch(const FCmdQueryOptions& Options)
 {
+	using namespace ProxyQuery;
+
 	TResult<FAuthToken> AuthToken = Authenticate(Options.Remote, 5 * 60);
 	if (!AuthToken.IsOk())
 	{
@@ -240,13 +242,18 @@ CmdQuerySearch(const FCmdQueryOptions& Options)
 		PendingDirectories.push_back(RootEntry);
 	}
 
+	struct FResultEntry : FEntry
+	{
+		FDirectoryListingEntry DirEntry;
+	};
+
 	struct FTaskContext
 	{
-		std::mutex Mutex;
-		FSemaphore ConnectionSemaphore = FSemaphore(8);
-		std::vector<std::string> FoundDirectories;
-		bool					 bParentThreadVerbose = false;
-		int32					 ParentThreadIndent = 0;
+		std::mutex				  Mutex;
+		FSemaphore				  ConnectionSemaphore = FSemaphore(8);
+		std::vector<FResultEntry> FoundEntries;
+		bool					  bParentThreadVerbose = false;
+		int32					  ParentThreadIndent   = 0;
 	};
 
 	FTaskGroup Tasks;
@@ -286,7 +293,6 @@ CmdQuerySearch(const FCmdQueryOptions& Options)
 
 		Response.Buffer.PushBack(0);
 
-		using namespace ProxyQuery;
 		TResult<FDirectoryListing> DirectoryListingResult = FDirectoryListing::FromJson((const char*)Response.Buffer.Data());
 		if (DirectoryListingResult.IsError())
 		{
@@ -296,24 +302,29 @@ CmdQuerySearch(const FCmdQueryOptions& Options)
 
 		const FDirectoryListing& DirectoryListing = DirectoryListingResult.GetData();
 
-		for (const FDirectoryListingEntry& Subdir : DirectoryListing.Entries)
+		for (const FDirectoryListingEntry& DirEntry : DirectoryListing.Entries)
 		{
 			const std::regex& RequiredPattern = SubdirPatterns[CurrentDepth];
 
-			if (std::regex_match(Subdir.Name, RequiredPattern, std::regex_constants::match_any))
+			if (std::regex_match(DirEntry.Name, RequiredPattern, std::regex_constants::match_any))
 			{
-				UNSYNC_VERBOSE2(L"Matched: '%hs'", Subdir.Name.c_str());
+				UNSYNC_VERBOSE2(L"Matched: '%hs'", DirEntry.Name.c_str());
 
 				FEntry NextEntry;
-				NextEntry.Path	= Path + "\\" + Subdir.Name;
+				NextEntry.Path	= Path + "\\" + DirEntry.Name;
 				NextEntry.Depth = CurrentDepth + 1;
 
 				if (NextEntry.Depth == SubdirPatterns.size())
 				{
-					{
-						std::lock_guard<std::mutex> LockGuard(Context.Mutex);
-						Context.FoundDirectories.push_back(NextEntry.Path);
-					}
+					std::lock_guard<std::mutex> LockGuard(Context.Mutex);
+
+					FResultEntry ResultEntry;
+					ResultEntry.Path	 = NextEntry.Path;
+					ResultEntry.Depth	 = NextEntry.Depth;
+					ResultEntry.DirEntry = DirEntry;
+
+					Context.FoundEntries.push_back(ResultEntry);
+
 					continue;
 				}
 
@@ -327,18 +338,30 @@ CmdQuerySearch(const FCmdQueryOptions& Options)
 	ExploreDirectory(RootPath, 0);
 	Tasks.wait();
 
-	std::vector<std::string> FoundDirectories = std::move(Context.FoundDirectories);
+	std::vector<FResultEntry>& ResultEntries = Context.FoundEntries;
 
-	std::sort(FoundDirectories.begin(), FoundDirectories.end());
+	std::sort(ResultEntries.begin(), ResultEntries.end(), [](const FResultEntry& A, const FResultEntry& B) { return A.Path < B.Path; });
 
-	LogPrintf(ELogLevel::MachineReadable, L"[\n");
-	for (size_t i = 0; i < FoundDirectories.size(); ++i)
+	LogPrintf(ELogLevel::MachineReadable, L"{\n");
+	LogPrintf(ELogLevel::MachineReadable, L"  \"root\": \"%hs\",\n", StringEscape(RootPath).c_str());
+	LogPrintf(ELogLevel::MachineReadable, L"  \"entries\": [\n");
+	for (size_t i = 0; i < ResultEntries.size(); ++i)
 	{
-		const std::string& ResultEntry	 = FoundDirectories[i];
-		const char*		   TrailingComma = i + 1 == FoundDirectories.size() ? "" : ",";
-		LogPrintf(ELogLevel::MachineReadable, L"  \"%hs\"%hs\n", StringEscape(ResultEntry).c_str(), TrailingComma);
+		const FResultEntry& ResultEntry	  = ResultEntries[i];
+		const char*			TrailingComma = i + 1 == ResultEntries.size() ? "" : ",";
+
+		std::string_view RelativePath = std::string_view(ResultEntry.Path).substr(RootPath.length() + 1);
+
+		LogPrintf(ELogLevel::MachineReadable,
+				  L"    { \"path\": \"%hs\", \"is_directory\": %hs, \"mtime\": %llu, \"size\": %llu }%hs\n",
+				  StringEscape(RelativePath).c_str(),
+				  ResultEntry.DirEntry.bDirectory ? "true" : "false",
+				  llu(ResultEntry.DirEntry.Mtime),
+				  llu(ResultEntry.DirEntry.Size),
+				  TrailingComma);
 	}
-	LogPrintf(ELogLevel::MachineReadable, L"]\n");
+	LogPrintf(ELogLevel::MachineReadable, L"  ]\n");
+	LogPrintf(ELogLevel::MachineReadable, L"}\n");
 
 	return 0;
 }
diff --git a/Engine/Source/Programs/Unsync/Private/UnsyncCore.cpp b/Engine/Source/Programs/Unsync/Private/UnsyncCore.cpp
--- a/Engine/Source/Programs/Unsync/Private/UnsyncCore.cpp
+++ b/Engine/Source/Programs/Unsync/Private/UnsyncCore.cpp
@@ -27,7 +27,7 @@ UNSYNC_THIRD_PARTY_INCLUDES_START
 #include <md5-sse2.h>
 UNSYNC_THIRD_PARTY_INCLUDES_END
 
-#define UNSYNC_VERSION_STR "1.0.57"
+#define UNSYNC_VERSION_STR "1.0.59"
 
 namespace unsync {
 
