unsync - Direct file download and search query improvements

* Derive output filename if one was not provided explicitly
* Include all encountered files and directories in the search query, not just leaf
* Bump version to 1.0.60

diff --git a/Engine/Source/Programs/Unsync/Private/UnsyncCmdQuery.cpp b/Engine/Source/Programs/Unsync/Private/UnsyncCmdQuery.cpp
--- a/Engine/Source/Programs/Unsync/Private/UnsyncCmdQuery.cpp
+++ b/Engine/Source/Programs/Unsync/Private/UnsyncCmdQuery.cpp
@@ -304,33 +304,30 @@ CmdQuerySearch(const FCmdQueryOptions& Options)
 
 		for (const FDirectoryListingEntry& DirEntry : DirectoryListing.Entries)
 		{
-			const std::regex& RequiredPattern = SubdirPatterns[CurrentDepth];
+			FEntry NextEntry;
+			NextEntry.Path	= Path + "\\" + DirEntry.Name;
+			NextEntry.Depth = CurrentDepth + 1;
 
-			if (std::regex_match(DirEntry.Name, RequiredPattern, std::regex_constants::match_any))
 			{
-				UNSYNC_VERBOSE2(L"Matched: '%hs'", DirEntry.Name.c_str());
-
-				FEntry NextEntry;
-				NextEntry.Path	= Path + "\\" + DirEntry.Name;
-				NextEntry.Depth = CurrentDepth + 1;
+				std::lock_guard<std::mutex> LockGuard(Context.Mutex);
 
-				if (NextEntry.Depth == SubdirPatterns.size())
-				{
-					std::lock_guard<std::mutex> LockGuard(Context.Mutex);
+				FResultEntry ResultEntry;
+				ResultEntry.Path	 = NextEntry.Path;
+				ResultEntry.Depth	 = NextEntry.Depth;
+				ResultEntry.DirEntry = DirEntry;
 
-					FResultEntry ResultEntry;
-					ResultEntry.Path	 = NextEntry.Path;
-					ResultEntry.Depth	 = NextEntry.Depth;
-					ResultEntry.DirEntry = DirEntry;
-
-					Context.FoundEntries.push_back(ResultEntry);
+				Context.FoundEntries.push_back(ResultEntry);
+			}
 
-					continue;
-				}
+			if (NextEntry.Depth > SubdirPatterns.size())
+			{
+				continue;
+			}
 
-				Tasks.run([ExploreDirectory, NextEntry]() {
-					ExploreDirectory(NextEntry.Path, NextEntry.Depth);
-				});
+			if (DirEntry.bDirectory && std::regex_match(DirEntry.Name, SubdirPatterns[CurrentDepth], std::regex_constants::match_any))
+			{
+				UNSYNC_VERBOSE2(L"Matched: '%hs'", DirEntry.Name.c_str());
+				Tasks.run([ExploreDirectory, NextEntry]() { ExploreDirectory(NextEntry.Path, NextEntry.Depth); });
 			}
 		}
 	};
@@ -380,9 +377,30 @@ CmdQueryFile(const FCmdQueryOptions& Options)
 
 	std::unique_ptr<FNativeFile> ResultWriter;
 
-	auto OutputCallback = [&ResultWriter, &Options](uint64 Size) -> FIOWriter& {
+	FPath OutputPath = Options.OutputPath;
+	if (OutputPath.empty())
+	{
+		FPath RequestPath = Options.Args[0];
+		if (RequestPath.has_filename())
+		{
+			OutputPath = RequestPath.filename();
+			OutputPath = GetAbsoluteNormalPath(OutputPath);
+		}
+		else
+		{
+			UNSYNC_ERROR(
+				L"Output could not be derived from the request string. "
+				L"Use `-o <filename>` command line argument to specify it explicitly.");
+			return 1;
+		}
+	}
+
+	UNSYNC_LOG(L"Output file: '%ls'", OutputPath.wstring().c_str());
+
+	auto OutputCallback = [&ResultWriter, &Options, OutputPath](uint64 Size) -> FIOWriter&
+	{
 		UNSYNC_LOG(L"Size: %llu bytes (%.3f MB)", llu(Size), SizeMb(Size));
-		ResultWriter = std::make_unique<FNativeFile>(Options.OutputPath, EFileMode::CreateWriteOnly, Size);
+		ResultWriter = std::make_unique<FNativeFile>(OutputPath, EFileMode::CreateWriteOnly, Size);
 		return *ResultWriter;
 	};
 
@@ -390,7 +408,7 @@ CmdQueryFile(const FCmdQueryOptions& Options)
 
 	if (Response.IsOk())
 	{
-		UNSYNC_LOG(L"Output written to file '%ls'", Options.OutputPath.wstring().c_str());
+		UNSYNC_LOG(L"Output written to file '%ls'", OutputPath.wstring().c_str());
 	}
 	else
 	{
diff --git a/Engine/Source/Programs/Unsync/Private/UnsyncCore.cpp b/Engine/Source/Programs/Unsync/Private/UnsyncCore.cpp
--- a/Engine/Source/Programs/Unsync/Private/UnsyncCore.cpp
+++ b/Engine/Source/Programs/Unsync/Private/UnsyncCore.cpp
@@ -27,7 +27,7 @@ UNSYNC_THIRD_PARTY_INCLUDES_START
 #include <md5-sse2.h>
 UNSYNC_THIRD_PARTY_INCLUDES_END
 
-#define UNSYNC_VERSION_STR "1.0.59"
+#define UNSYNC_VERSION_STR "1.0.60"
 
 namespace unsync {
 
@@ -3015,7 +3015,7 @@ DownloadFileIfNewer(const FRemoteDesc& RemoteDesc, const FPath& Source, const FP
 			return false;
 		}
 
-		SetFileMtime(Target, SourceEntry->Mtime);
+		SetFileMtime(Target, SourceEntry->Mtime, /*allow in dry run*/ true);
 	}
 
 	return true;
diff --git a/Engine/Source/Programs/Unsync/Private/UnsyncFile.cpp b/Engine/Source/Programs/Unsync/Private/UnsyncFile.cpp
--- a/Engine/Source/Programs/Unsync/Private/UnsyncFile.cpp
+++ b/Engine/Source/Programs/Unsync/Private/UnsyncFile.cpp
@@ -601,9 +601,9 @@ GetFileAttrib(const FPath& Path, FFileAttributeCache* AttribCache)
 }
 
 bool
-SetFileMtime(const FPath& Path, uint64 Mtime)
+SetFileMtime(const FPath& Path, uint64 Mtime, bool bAllowInDryRun)
 {
-	UNSYNC_ASSERT(!GDryRun);
+	UNSYNC_ASSERT(!GDryRun || bAllowInDryRun);
 
 	FPath ExtendedPath = MakeExtendedAbsolutePath(Path);
 
@@ -809,7 +809,7 @@ SetFileReadOnly(const FPath& path, bool bReadOnly)
 }
 
 bool
-SetFileMtime(const FPath& path, uint64 mtime)
+SetFileMtime(const FPath& path, uint64 mtime, bool bAllowInDryRun)
 {
 	UNSYNC_WARNING(L"SetFileMtime() is not implemented");
 	return false;
diff --git a/Engine/Source/Programs/Unsync/Private/UnsyncFile.h b/Engine/Source/Programs/Unsync/Private/UnsyncFile.h
--- a/Engine/Source/Programs/Unsync/Private/UnsyncFile.h
+++ b/Engine/Source/Programs/Unsync/Private/UnsyncFile.h
@@ -445,7 +445,7 @@ struct FFileAttributeCache
 };
 
 FFileAttributes GetFileAttrib(const FPath& Path, FFileAttributeCache* AttribCache = nullptr);
-bool			SetFileMtime(const FPath& Path, uint64 Mtime);
+bool			SetFileMtime(const FPath& Path, uint64 Mtime, bool bAllowInDryRun = false);
 bool			SetFileReadOnly(const FPath& Path, bool ReadOnly);
 bool			IsDirectory(const FPath& Path);
 bool			PathExists(const FPath& Path);
