unsync - Always use the specified proxy server for performing authentication instead of auto-selected closest, fix few other minor issues

* Use correct TLS subject after auto-selecting a proxy (previously the root proxy address was used)

diff --git a/Engine/Source/Programs/Unsync/Private/UnsyncAuth.cpp b/Engine/Source/Programs/Unsync/Private/UnsyncAuth.cpp
--- a/Engine/Source/Programs/Unsync/Private/UnsyncAuth.cpp
+++ b/Engine/Source/Programs/Unsync/Private/UnsyncAuth.cpp
@@ -568,7 +568,7 @@ RefreshAuthToken(const FAuthDesc& AuthDesc, const FAuthToken& PreviousToken)
 std::string
 GenerateTokenId(const FRemoteDesc& RemoteDesc)
 {
-	FHash128 Hash = HashBlake3String<FHash128>(RemoteDesc.HostAddress);
+	FHash128 Hash = HashBlake3String<FHash128>(RemoteDesc.GetLoginAddress());
 	return HashToHexString(Hash);
 }
 
@@ -847,7 +847,10 @@ GetAuthenticationDesc(const FRemoteDesc& RemoteDesc)
 		return AppError(L"Authentication is only implemented for UNSYNC protocol");
 	}
 
-	TResult<ProxyQuery::FHelloResponse> HelloResponseResult = ProxyQuery::Hello(RemoteDesc, /*bAnonymous*/ true);
+	FRemoteDesc LoginRemoteDesc = RemoteDesc;
+	LoginRemoteDesc.HostAddress = RemoteDesc.GetLoginAddress();
+
+	TResult<ProxyQuery::FHelloResponse> HelloResponseResult = ProxyQuery::Hello(LoginRemoteDesc, /*bAnonymous*/ true);
 	if (HelloResponseResult.IsError())
 	{
 		return MoveError<FAuthDesc>(HelloResponseResult);
diff --git a/Engine/Source/Programs/Unsync/Private/UnsyncBuffer.cpp b/Engine/Source/Programs/Unsync/Private/UnsyncBuffer.cpp
--- a/Engine/Source/Programs/Unsync/Private/UnsyncBuffer.cpp
+++ b/Engine/Source/Programs/Unsync/Private/UnsyncBuffer.cpp
@@ -65,6 +65,12 @@ FBuffer::Append(const FBuffer& Other)
 	Append(Other.Data(), Other.Size());
 }
 
+void
+FBuffer::Append(const FBufferView& Other)
+{
+	Append(Other.Data, Other.Size);
+}
+
 void
 FBuffer::Append(const uint8* AppendData, uint64 AppendSize)
 {
diff --git a/Engine/Source/Programs/Unsync/Private/UnsyncBuffer.h b/Engine/Source/Programs/Unsync/Private/UnsyncBuffer.h
--- a/Engine/Source/Programs/Unsync/Private/UnsyncBuffer.h
+++ b/Engine/Source/Programs/Unsync/Private/UnsyncBuffer.h
@@ -18,6 +18,14 @@ struct FMutBufferView
 {
 	uint8* Data = nullptr;
 	uint64 Size = 0;
+
+	operator FBufferView() const
+	{ 
+		FBufferView Result;
+		Result.Data = Data;
+		Result.Size = Size;
+		return Result;
+	}
 };
 
 class FBuffer
@@ -49,6 +57,7 @@ public:
 
 	void Append(const uint8* AppendData, uint64 AppendSize);
 	void Append(const FBuffer& Other);
+	void Append(const FBufferView& Other);
 	void PushBack(uint8 Value) { Append(&Value, 1); }
 	void Resize(uint64 NewSize);
 	void Clear() { Resize(0); }
diff --git a/Engine/Source/Programs/Unsync/Private/UnsyncCmdPush.cpp b/Engine/Source/Programs/Unsync/Private/UnsyncCmdPush.cpp
--- a/Engine/Source/Programs/Unsync/Private/UnsyncCmdPush.cpp
+++ b/Engine/Source/Programs/Unsync/Private/UnsyncCmdPush.cpp
@@ -70,9 +70,7 @@ CmdPush(const FCmdPushOptions& Options)
 
 		const bool bUseTls = Options.Remote.bTlsEnable;
 
-		FTlsClientSettings TlsSettings;
-		TlsSettings.Subject			   = Options.Remote.HostAddress.c_str();
-		TlsSettings.bVerifyCertificate = Options.Remote.bTlsVerifyCertificate;
+		FTlsClientSettings TlsSettings = Options.Remote.GetTlsClientSettings();
 
 		if (Options.Remote.StorageNamespace.empty())
 		{
diff --git a/Engine/Source/Programs/Unsync/Private/UnsyncCmdQuery.cpp b/Engine/Source/Programs/Unsync/Private/UnsyncCmdQuery.cpp
--- a/Engine/Source/Programs/Unsync/Private/UnsyncCmdQuery.cpp
+++ b/Engine/Source/Programs/Unsync/Private/UnsyncCmdQuery.cpp
@@ -203,9 +203,7 @@ CmdQuerySearch(const FCmdQueryOptions& Options)
 
 	auto CreateConnection = [Remote = Options.Remote]
 	{
-		FTlsClientSettings TlsSettings;
-		TlsSettings.Subject			   = Remote.HostAddress.data();
-		TlsSettings.bVerifyCertificate = Remote.bTlsVerifyCertificate;
+		FTlsClientSettings TlsSettings = Remote.GetTlsClientSettings();
 		return new FHttpConnection(Remote.HostAddress, Remote.HostPort, &TlsSettings);
 	};
 
diff --git a/Engine/Source/Programs/Unsync/Private/UnsyncCore.cpp b/Engine/Source/Programs/Unsync/Private/UnsyncCore.cpp
--- a/Engine/Source/Programs/Unsync/Private/UnsyncCore.cpp
+++ b/Engine/Source/Programs/Unsync/Private/UnsyncCore.cpp
@@ -27,7 +27,7 @@ UNSYNC_THIRD_PARTY_INCLUDES_START
 #include <md5-sse2.h>
 UNSYNC_THIRD_PARTY_INCLUDES_END
 
-#define UNSYNC_VERSION_STR "1.0.60"
+#define UNSYNC_VERSION_STR "1.0.61-dev"
 
 namespace unsync {
 
diff --git a/Engine/Source/Programs/Unsync/Private/UnsyncHttp.cpp b/Engine/Source/Programs/Unsync/Private/UnsyncHttp.cpp
--- a/Engine/Source/Programs/Unsync/Private/UnsyncHttp.cpp
+++ b/Engine/Source/Programs/Unsync/Private/UnsyncHttp.cpp
@@ -491,12 +491,15 @@ FHttpConnection::FHttpConnection(const std::string_view InHostAddress, uint16 In
 {
 	if (InTlsSettings)
 	{
-		TlsSubject			  = InTlsSettings->Subject ? InTlsSettings->Subject : std::string();
+		UNSYNC_ASSERT(!InTlsSettings->Subject.empty());
+
+		TlsSubject			  = std::string(InTlsSettings->Subject);
 		bTlsVerifyCertificate = InTlsSettings->bVerifyCertificate;
-		if (InTlsSettings->CacertData)
+		bTlsVerifySubject	  = InTlsSettings->bVerifySubject;
+		if (InTlsSettings->CACert.Data)
 		{
 			TlsCacert = std::make_shared<FBuffer>();
-			TlsCacert->Append(InTlsSettings->CacertData, InTlsSettings->CacertSize);
+			TlsCacert->Append(InTlsSettings->CACert.Data, InTlsSettings->CACert.Size);
 		}
 	}
 }
@@ -506,6 +509,7 @@ FHttpConnection::FHttpConnection(const FHttpConnection& Other)
 , HostPort(Other.HostPort)
 , bUseTls(Other.bUseTls)
 , bKeepAlive(Other.bKeepAlive)
+, bTlsVerifySubject(Other.bTlsVerifySubject)
 , TlsSubject(Other.TlsSubject)
 , bTlsVerifyCertificate(Other.bTlsVerifyCertificate)
 , TlsCacert(Other.TlsCacert)
@@ -522,17 +526,14 @@ FHttpConnection
 FHttpConnection::CreateDefaultHttps(const std::string_view InHostAddress, uint16 Port)
 {
 	FTlsClientSettings TlsSettings;
-	TlsSettings.Subject			   = InHostAddress.data();
-	TlsSettings.bVerifyCertificate = true;
+	TlsSettings.Subject = InHostAddress.data();
 	return FHttpConnection(InHostAddress, Port, &TlsSettings);
 }
 
 FHttpConnection
 FHttpConnection::CreateDefaultHttps(const FRemoteDesc& RemoteDesc)
 {
-	FTlsClientSettings TlsSettings;
-	TlsSettings.Subject			   = RemoteDesc.HostAddress.data();
-	TlsSettings.bVerifyCertificate = RemoteDesc.bTlsVerifyCertificate;
+	FTlsClientSettings TlsSettings = RemoteDesc.GetTlsClientSettings();
 	return FHttpConnection(RemoteDesc.HostAddress, RemoteDesc.HostPort, &TlsSettings);
 }
 
@@ -558,16 +559,12 @@ FHttpConnection::Open()
 	{
 		FTlsClientSettings ClientSettings;
 		ClientSettings.bVerifyCertificate = bTlsVerifyCertificate;
-
-		if (bTlsVerifyCertificate && !TlsSubject.empty())
-		{
-			ClientSettings.Subject = TlsSubject.c_str();
-		}
+		ClientSettings.bVerifySubject	  = bTlsVerifySubject;
+		ClientSettings.Subject			  = TlsSubject;
 
 		if (TlsCacert && !TlsCacert->Empty())
 		{
-			ClientSettings.CacertData = TlsCacert->Data();
-			ClientSettings.CacertSize = TlsCacert->Size();
+			ClientSettings.CACert = TlsCacert->View();
 		}
 
 		FSocketTls* TlsSocket = new FSocketTls(RawSocketHandle, ClientSettings);
diff --git a/Engine/Source/Programs/Unsync/Private/UnsyncHttp.h b/Engine/Source/Programs/Unsync/Private/UnsyncHttp.h
--- a/Engine/Source/Programs/Unsync/Private/UnsyncHttp.h
+++ b/Engine/Source/Programs/Unsync/Private/UnsyncHttp.h
@@ -78,6 +78,7 @@ struct FHttpConnection
 
 	bool bKeepAlive = true;
 
+	bool					 bTlsVerifySubject = true;
 	std::string				 TlsSubject;
 	bool					 bTlsVerifyCertificate = true;
 	std::shared_ptr<FBuffer> TlsCacert;
diff --git a/Engine/Source/Programs/Unsync/Private/UnsyncMain.cpp b/Engine/Source/Programs/Unsync/Private/UnsyncMain.cpp
--- a/Engine/Source/Programs/Unsync/Private/UnsyncMain.cpp
+++ b/Engine/Source/Programs/Unsync/Private/UnsyncMain.cpp
@@ -661,12 +661,13 @@ InnerMain(int Argc, char** Argv)
 	if (bAllowInsecureTls)
 	{
 		RemoteDesc.bTlsVerifyCertificate = false;
-		RemoteDesc.TlsSubject			 = {};
+		RemoteDesc.bTlsVerifySubject	 = false;
 		UNSYNC_WARNING(L"Remote server certificate verification is disabled.");
 	}
 	else
 	{
 		RemoteDesc.bTlsVerifyCertificate = true;
+		RemoteDesc.bTlsVerifySubject	 = true;
 	}
 
 	if (bNoOutputValidation)
@@ -701,6 +702,11 @@ InnerMain(int Argc, char** Argv)
 		}
 	}
 
+	if (bShouldLogin)
+	{
+		RemoteDesc.LoginAddress = RemoteDesc.HostAddress;
+	}
+
 	FRemoteDesc RootRemoteDesc = RemoteDesc;
 
 	if (!bNoProxySelect
diff --git a/Engine/Source/Programs/Unsync/Private/UnsyncProxy.cpp b/Engine/Source/Programs/Unsync/Private/UnsyncProxy.cpp
--- a/Engine/Source/Programs/Unsync/Private/UnsyncProxy.cpp
+++ b/Engine/Source/Programs/Unsync/Private/UnsyncProxy.cpp
@@ -44,14 +44,7 @@ FProxy::FProxy(const FRemoteDesc& RemoteDesc, const FRemoteProtocolFeatures& InF
 {
 	UNSYNC_ASSERT(InRequestMap);
 
-	FTlsClientSettings TlsSettings;
-	TlsSettings.bVerifyCertificate = RemoteDesc.bTlsVerifyCertificate;
-	TlsSettings.Subject			   = RemoteDesc.TlsSubject.empty() ? nullptr : RemoteDesc.TlsSubject.c_str();
-	if (RemoteDesc.TlsCacert)
-	{
-		TlsSettings.CacertData = RemoteDesc.TlsCacert->Data();
-		TlsSettings.CacertSize = RemoteDesc.TlsCacert->Size();
-	}
+	FTlsClientSettings TlsSettings = RemoteDesc.GetTlsClientSettings();
 
 	if (RemoteDesc.Protocol == EProtocolFlavor::Jupiter)
 	{
@@ -585,9 +578,7 @@ ProxyQuery::DownloadFile(const FRemoteDesc& Remote, const std::string& Path, Pro
 {
 	auto CreateConnection = [Remote]
 	{
-		FTlsClientSettings TlsSettings;
-		TlsSettings.Subject			   = Remote.HostAddress.data();
-		TlsSettings.bVerifyCertificate = Remote.bTlsVerifyCertificate;
+		FTlsClientSettings TlsSettings = Remote.GetTlsClientSettings();
 		return new FHttpConnection(Remote.HostAddress, Remote.HostPort, &TlsSettings);
 	};
 
diff --git a/Engine/Source/Programs/Unsync/Private/UnsyncRemote.cpp b/Engine/Source/Programs/Unsync/Private/UnsyncRemote.cpp
--- a/Engine/Source/Programs/Unsync/Private/UnsyncRemote.cpp
+++ b/Engine/Source/Programs/Unsync/Private/UnsyncRemote.cpp
@@ -172,8 +172,8 @@ FRemoteDesc::FromUrl(std::string_view Url)
 	Result.HostPort	   = HostPort;
 
 	Result.bTlsEnable			 = bUseTls;
-	Result.bTlsVerifyCertificate = Result.bTlsEnable;
-	Result.TlsSubject			 = HostAddress;
+	Result.bTlsVerifyCertificate = bUseTls;
+	Result.bTlsVerifySubject	 = bUseTls;
 
 	return ResultOk(Result);
 }
@@ -184,17 +184,18 @@ FRemoteDesc::GetTlsClientSettings() const
 	FTlsClientSettings Result = {};
 	if (bTlsEnable)
 	{
-		Result.Subject			  = TlsSubject.empty() ? nullptr : TlsSubject.c_str();
+		Result.Subject			  = GetTlsSubject();
 		Result.bVerifyCertificate = bTlsVerifyCertificate;
+		Result.bVerifySubject	  = bTlsVerifySubject;
 		if (TlsCacert)
 		{
-			Result.CacertData = TlsCacert->Data();
-			Result.CacertSize = TlsCacert->Size();
+			Result.CACert = TlsCacert->View();
 		}
 	}
 	else
 	{
 		Result.bVerifyCertificate = false;
+		Result.bVerifySubject	  = false;
 	}
 	return Result;
 }
diff --git a/Engine/Source/Programs/Unsync/Private/UnsyncRemote.h b/Engine/Source/Programs/Unsync/Private/UnsyncRemote.h
--- a/Engine/Source/Programs/Unsync/Private/UnsyncRemote.h
+++ b/Engine/Source/Programs/Unsync/Private/UnsyncRemote.h
@@ -43,10 +43,16 @@ struct FRemoteDesc
 
 	bool					 bTlsEnable			   = true;	// Prefer TLS, if supported by protocol and remote server
 	bool					 bTlsVerifyCertificate = true;	// Disabling this allows self-signed certificates
-	std::string				 TlsSubject;					// Use host by default
+	bool					 bTlsVerifySubject	   = true;	// Disabling this is insecure, but may be useful during development
+	std::string				 TlsSubjectOverride;			// Use host address if empty (default)
 	std::shared_ptr<FBuffer> TlsCacert;	 // Custom CA to use for server certificate validation (system root CA is used by default)
 
+	const std::string& GetTlsSubject() const { return TlsSubjectOverride.length() ? TlsSubjectOverride : HostAddress; }
+
 	bool bAuthenticationRequired = false;
+	std::string LoginAddress; // Optional address of the server used for login requests. If empty, then HostAddress is used.
+
+	const std::string& GetLoginAddress() const { return LoginAddress.length() ? LoginAddress : HostAddress; }
 
 	uint32 RecvTimeoutSeconds = 0;
 
diff --git a/Engine/Source/Programs/Unsync/Private/UnsyncSocket.cpp b/Engine/Source/Programs/Unsync/Private/UnsyncSocket.cpp
--- a/Engine/Source/Programs/Unsync/Private/UnsyncSocket.cpp
+++ b/Engine/Source/Programs/Unsync/Private/UnsyncSocket.cpp
@@ -338,9 +338,9 @@ FSocketTls::FSocketTls(FSocketHandle InHandle, FTlsClientSettings ClientSettings
 
 	tls_config* TlsCfg = tls_config_new();
 
-	if (ClientSettings.CacertData)
+	if (ClientSettings.CACert.Data)
 	{
-		tls_config_set_ca_mem(TlsCfg, ClientSettings.CacertData, (size_t)ClientSettings.CacertSize);
+		tls_config_set_ca_mem(TlsCfg, ClientSettings.CACert.Data, (size_t)ClientSettings.CACert.Size);
 	}
 	else
 	{
@@ -353,7 +353,7 @@ FSocketTls::FSocketTls(FSocketHandle InHandle, FTlsClientSettings ClientSettings
 		tls_config_insecure_noverifycert(TlsCfg);
 	}
 
-	if (!ClientSettings.Subject)
+	if (!ClientSettings.bVerifySubject)
 	{
 		tls_config_insecure_noverifyname(TlsCfg);
 	}
@@ -371,7 +371,9 @@ FSocketTls::FSocketTls(FSocketHandle InHandle, FTlsClientSettings ClientSettings
 
 	if (Err == 0)
 	{
-		Err = tls_connect_socket(TlsCtx, (int)Handle, ClientSettings.Subject);
+		UNSYNC_ASSERT(!ClientSettings.Subject.empty());
+		std::string TlsSubject(ClientSettings.Subject);
+		Err = tls_connect_socket(TlsCtx, (int)Handle, TlsSubject.c_str());
 	}
 
 	if (Err == 0)
diff --git a/Engine/Source/Programs/Unsync/Private/UnsyncSocket.h b/Engine/Source/Programs/Unsync/Private/UnsyncSocket.h
--- a/Engine/Source/Programs/Unsync/Private/UnsyncSocket.h
+++ b/Engine/Source/Programs/Unsync/Private/UnsyncSocket.h
@@ -5,6 +5,7 @@
 #include <stdint.h>
 
 #include "UnsyncCommon.h"
+#include "UnsyncBuffer.h"
 
 struct tls;
 
@@ -103,10 +104,10 @@ struct FSocketRaw : FSocketBase
 
 struct FTlsClientSettings
 {
-	const char*	 Subject			= nullptr;
-	bool		 bVerifyCertificate = true;
-	const uint8* CacertData			= nullptr;
-	uint64		 CacertSize			= 0;
+	std::string_view Subject			= {};
+	FBufferView		 CACert				= {};
+	bool			 bVerifyCertificate = true;
+	bool			 bVerifySubject		= true;
 };
 
 struct FSocketTls : FSocketBase
