unsync - Add GetRelativePath() that handles UNC paths

* Use instead of std::filesystem::relative() when possible, to avoid filesystem operations (lexically_relative() does not handle the case where only one of the input paths is UNC)

diff --git a/Engine/Source/Programs/Unsync/Private/UnsyncCommon.h b/Engine/Source/Programs/Unsync/Private/UnsyncCommon.h
--- a/Engine/Source/Programs/Unsync/Private/UnsyncCommon.h
+++ b/Engine/Source/Programs/Unsync/Private/UnsyncCommon.h
@@ -4,6 +4,7 @@
 
 #include <stdint.h>
 #include <cstring>
+#include <string_view>
 #include <filesystem>
 
 namespace unsync {
@@ -21,6 +22,7 @@ using int32 = int32_t;
 using int64 = int64_t;
 
 using FPath = std::filesystem::path;
+using FPathStringView = std::basic_string_view<FPath::value_type>;
 
 #define UNSYNC_DISALLOW_COPY_ASSIGN(T) \
 	T(const T&) = delete;              \
diff --git a/Engine/Source/Programs/Unsync/Private/UnsyncCore.cpp b/Engine/Source/Programs/Unsync/Private/UnsyncCore.cpp
--- a/Engine/Source/Programs/Unsync/Private/UnsyncCore.cpp
+++ b/Engine/Source/Programs/Unsync/Private/UnsyncCore.cpp
@@ -1811,14 +1811,16 @@ CreateDirectoryManifest(const FPath& Root, uint32 BlockSize, FAlgorithmOptions A
 			continue;
 		}
 
-		FPath RelativePath = std::filesystem::relative(Dir.path(), Root);
+		FPath RelativePath = GetRelativePath(Dir.path(), Root);
 
 		if (RelativePath.native().starts_with(UnsyncDirName.native()))
 		{
 			continue;
 		}
 
-		UNSYNC_VERBOSE2(L"Found '%ls'", RelativePath.wstring().c_str());
+		const std::wstring PathKey = RelativePath.wstring();
+
+		UNSYNC_VERBOSE2(L"Found '%ls'", PathKey.c_str());
 
 		FFileManifest FileManifest;
 
@@ -1827,11 +1829,9 @@ CreateDirectoryManifest(const FPath& Root, uint32 BlockSize, FAlgorithmOptions A
 		FileManifest.CurrentPath = Dir.path();
 		FileManifest.BlockSize	 = BlockSize;
 
-		std::wstring Key = RelativePath.wstring();
-
 		{
 			std::lock_guard<std::mutex> LockGuard(ResultMutex);
-			Result.Files[Key] = std::move(FileManifest);
+			Result.Files[PathKey] = std::move(FileManifest);
 		}
 
 		if (BlockSize)
@@ -1843,7 +1843,7 @@ CreateDirectoryManifest(const FPath& Root, uint32 BlockSize, FAlgorithmOptions A
 				UNSYNC_VERBOSE(L"Computing blocks for '%ls' (%.2f MB)", FilePath.wstring().c_str(), double(File->GetSize()) / (1 << 20));
 
 				Semaphore.Acquire();
-				TaskGroup.run([&Semaphore, &ResultMutex, &Result, File = std::move(File), Key = std::move(Key), BlockSize, Algorithm]() {
+				TaskGroup.run([&Semaphore, &ResultMutex, &Result, File = std::move(File), Key = std::move(PathKey), BlockSize, Algorithm]() {
 					FComputeMacroBlockParams MacroBlocks;
 
 					// TODO: macro block generation is only implemented for variable chunk mode
@@ -3228,7 +3228,7 @@ DeleteOldFilesInDirectory(FPath& Path, uint32 MaxFilesToKeep)
 	{
 		const FEntry& Oldest = Entries.back();
 
-		std::wstring PathStr = RemoveExtendedPathPrefix(Oldest.Path).wstring();
+		std::wstring PathStr = FPath(RemoveExtendedPathPrefix(Oldest.Path)).wstring();
 
 		if (GDryRun)
 		{
diff --git a/Engine/Source/Programs/Unsync/Private/UnsyncFile.cpp b/Engine/Source/Programs/Unsync/Private/UnsyncFile.cpp
--- a/Engine/Source/Programs/Unsync/Private/UnsyncFile.cpp
+++ b/Engine/Source/Programs/Unsync/Private/UnsyncFile.cpp
@@ -55,28 +55,47 @@ MakeExtendedAbsolutePath(const FPath& InAbsolutePath)
 #endif // UNSYNC_PLATFORM_WINDOWS
 }
 
-FPath
+FPathStringView
 RemoveExtendedPathPrefix(const FPath& InPath)
 {
+	FPathStringView InPathString = InPath.native();
 #if UNSYNC_PLATFORM_WINDOWS
-	const std::wstring& InPathString = InPath.native();
 	if (InPathString.starts_with(L"\\\\?\\UNC\\"))
 	{
-		return FPath(InPathString.substr(8));
+		return InPathString.substr(8);
 	}
 	else if (InPathString.starts_with(L"\\\\?\\"))
 	{
-		return FPath(InPathString.substr(4));
+		return InPathString.substr(4);
 	}
 	else
 	{
-		return InPath;
+		return InPathString;
 	}
 #else // UNSYNC_PLATFORM_WINDOWS
-	return InPath;
+	return InPathString;
 #endif // UNSYNC_PLATFORM_WINDOWS
 }
 
+FPath
+GetRelativePath(const FPath& Path, const FPath& Base)
+{
+	// Try a trivial case first, without touching the filesystem
+	FPathStringView PathView = RemoveExtendedPathPrefix(Path);
+	FPathStringView BaseView = RemoveExtendedPathPrefix(Base);
+	if (PathView.starts_with(BaseView))
+	{
+		FPathStringView RelativePath = PathView.substr(BaseView.length());
+		while (RelativePath.starts_with(FPath::preferred_separator))
+		{
+			RelativePath = RelativePath.substr(1);
+		}
+		return FPath(RelativePath);
+	}
+
+	return std::filesystem::relative(Path, Base);
+}
+
 #if UNSYNC_PLATFORM_WINDOWS
 inline uint64
 MakeU64(FILETIME Ft)
diff --git a/Engine/Source/Programs/Unsync/Private/UnsyncFile.h b/Engine/Source/Programs/Unsync/Private/UnsyncFile.h
--- a/Engine/Source/Programs/Unsync/Private/UnsyncFile.h
+++ b/Engine/Source/Programs/Unsync/Private/UnsyncFile.h
@@ -455,6 +455,7 @@ bool			FileRename(const FPath& From, const FPath& To, std::error_code& OutErrorC
 bool			FileCopy(const FPath& From, const FPath& To, std::error_code& OutErrorCode);
 bool			FileCopyOverwrite(const FPath& From, const FPath& To, std::error_code& OutErrorCode);
 bool			FileRemove(const FPath& Path, std::error_code& OutErrorCode);
+FPath			GetRelativePath(const FPath& Path, const FPath& Base);
 
 // Returns number of bytes that can be written to the given path.
 // Returns ~0ull if the available space could not be determined.
@@ -474,6 +475,6 @@ FPath MakeExtendedAbsolutePath(const FPath& InAbsolutePath);
 
 // Removes `\\?\` or `\\?\UNC\` prefix from a given path.
 // Returns original path on non-Windows.
-FPath RemoveExtendedPathPrefix(const FPath& InPath);
+FPathStringView RemoveExtendedPathPrefix(const FPath& InPath);
 
 }  // namespace unsync
