unsync - Several Linux compatibility fixes (more still required for full support)

* Implement SetFileMtime and SetFileReadOnly
* Only allow --flag command line syntax, not /flag

diff --git a/Engine/Source/Programs/Unsync/Private/UnsyncCore.cpp b/Engine/Source/Programs/Unsync/Private/UnsyncCore.cpp
--- a/Engine/Source/Programs/Unsync/Private/UnsyncCore.cpp
+++ b/Engine/Source/Programs/Unsync/Private/UnsyncCore.cpp
@@ -3759,12 +3759,11 @@ SyncDirectory(const FSyncDirectoryOptions& SyncOptions)
 			ProxyPool.Invalidate();
 		}
 	}
-
-	if (!ProxyPool.IsValid())
+	else
 	{
 		// TODO: bail out if remote connection is required for the download,
 		// such as when downloading data purely from Jupiter.
-		UNSYNC_WARNING(L"Attempting to sync without remote server connection");
+		UNSYNC_VERBOSE(L"Attempting to sync without remote server connection");
 	}
 
 	LogGlobalStatus(L"Copying files");
diff --git a/Engine/Source/Programs/Unsync/Private/UnsyncFile.cpp b/Engine/Source/Programs/Unsync/Private/UnsyncFile.cpp
--- a/Engine/Source/Programs/Unsync/Private/UnsyncFile.cpp
+++ b/Engine/Source/Programs/Unsync/Private/UnsyncFile.cpp
@@ -24,6 +24,12 @@ namespace unsync {
 
 bool GForceBufferedFiles = false;
 
+// Windows epoch : 1601-01-01T00:00:00Z
+// Unix epoch    : 1970-01-01T00:00:00Z
+static constexpr uint64 SECONDS_BETWEEN_WINDOWS_AND_UNIX = 11'644'473'600ull;
+static constexpr uint64 NANOS_PER_WINDOWS_TICK			 = 100ull;
+static constexpr uint64 WINDOWS_TICKS_PER_SECOND		 = 1'000'000'000ull / NANOS_PER_WINDOWS_TICK;  // each tick is 100ns
+
 // Returns extended absolute path of a form \\?\D:\verylongpath or \\?\UNC\servername\verylongpath
 // Expects an absolute path input. Returns original path on non-Windows.
 // https://docs.microsoft.com/en-us/windows/win32/fileio/maximum-file-path-limitation
@@ -77,6 +83,29 @@ RemoveExtendedPathPrefix(const FPath& InPath)
 #endif // UNSYNC_PLATFORM_WINDOWS
 }
 
+std::filesystem::file_time_type FromWindowsFileTime(uint64 Ticks)
+{
+	using FileTimeDuration = std::filesystem::file_time_type::duration;
+
+	uint64 RawSeconds = Ticks / WINDOWS_TICKS_PER_SECOND;
+	uint64 RawSubsecondTicks = Ticks - (RawSeconds * WINDOWS_TICKS_PER_SECOND);
+	uint64 RawSubsecondNanos = RawSubsecondTicks * NANOS_PER_WINDOWS_TICK;
+
+#if UNSYNC_PLATFORM_WINDOWS
+	FileTimeDuration Seconds = std::chrono::duration_cast<FileTimeDuration>(std::chrono::seconds(RawSeconds));
+#else	// UNSYNC_PLATFORM_WINDOWS
+	FileTimeDuration Seconds = std::chrono::seconds(RawSeconds - SECONDS_BETWEEN_WINDOWS_AND_UNIX);
+#endif	// UNSYNC_PLATFORM_WINDOWS
+
+	FileTimeDuration SubsecondNanos = std::chrono::duration_cast<FileTimeDuration>(std::chrono::nanoseconds(RawSubsecondNanos));
+
+	FileTimeDuration DurationFromNativeEpoch = Seconds + SubsecondNanos;
+
+	std::filesystem::file_time_type Result(DurationFromNativeEpoch);
+
+	return Result;
+}
+
 FPath
 GetRelativePath(const FPath& Path, const FPath& Base)
 {
@@ -619,58 +648,6 @@ GetFileAttrib(const FPath& Path, FFileAttributeCache* AttribCache)
 	return Result;
 }
 
-bool
-SetFileMtime(const FPath& Path, uint64 Mtime, bool bAllowInDryRun)
-{
-	UNSYNC_ASSERT(!GDryRun || bAllowInDryRun);
-
-	FPath ExtendedPath = MakeExtendedAbsolutePath(Path);
-
-	HANDLE Fh = CreateFileW(ExtendedPath.c_str(), FILE_WRITE_ATTRIBUTES, 0, nullptr, OPEN_EXISTING, FILE_ATTRIBUTE_NORMAL, nullptr);
-	if (Fh != INVALID_HANDLE_VALUE)
-	{
-		FILETIME C, A, w;
-		if (GetFileTime(Fh, &C, &A, &w))
-		{
-			w.dwHighDateTime = uint32(Mtime >> 32);
-			w.dwLowDateTime	 = uint32(Mtime);
-			bool bResult	 = SetFileTime(Fh, &C, &A, &w);
-			CloseHandle(Fh);
-			return bResult;
-		}
-	}
-	return false;
-}
-
-bool
-SetFileReadOnly(const FPath& Path, bool bReadOnly)
-{
-	UNSYNC_ASSERT(!GDryRun);
-
-	FPath ExtendedPath = MakeExtendedAbsolutePath(Path);
-
-	uint32 OldAttributes = GetFileAttributesW(ExtendedPath.c_str());
-	uint32 NewAttributes = OldAttributes;
-
-	if (bReadOnly)
-	{
-		NewAttributes |= FILE_ATTRIBUTE_READONLY;
-	}
-	else
-	{
-		NewAttributes &= ~FILE_ATTRIBUTE_READONLY;
-	}
-
-	if (NewAttributes == OldAttributes)
-	{
-		return true;
-	}
-	else
-	{
-		return SetFileAttributesW(ExtendedPath.c_str(), NewAttributes);
-	}
-}
-
 uint64
 ToWindowsFileTime(const std::filesystem::file_time_type& T)
 {
@@ -789,63 +766,46 @@ FUnixFile::Write(const void* data, uint64 DestOffset, uint64 WriteSize)
 }
 
 FFileAttributes
-GetFileAttrib(const FPath& path, FFileAttributeCache* AttribCache)
+GetFileAttrib(const FPath& Path, FFileAttributeCache* AttribCache)
 {
-	FFileAttributes result;
+	FFileAttributes Result;
 
 	if (AttribCache)
 	{
-		auto it = AttribCache->Map.find(path);
+		auto it = AttribCache->Map.find(Path);
 		if (it != AttribCache->Map.end())
 		{
-			result = it->second;
-			return result;
+			Result = it->second;
+			return Result;
 		}
 	}
 
 	// TODO: could potentially use std::filesystem::directory_entry for this on all platforms
 
-	std::error_code ec	  = {};
-	auto			entry = std::filesystem::directory_entry(path, ec);
+	std::error_code ErrorCode	  = {};
+	auto			Entry	  = std::filesystem::directory_entry(Path, ErrorCode);
 
-	if (ec.value() == 0)
+	if (!ErrorCode)
 	{
-		result.bDirectory = entry.is_directory();
-		result.Size		  = result.bDirectory ? 0 : entry.file_size();
-		result.Mtime	  = ToWindowsFileTime(entry.last_write_time());
-		result.bReadOnly  = false;	// TODO
-		result.bValid	  = true;
-	}
-
-	return result;
-}
+		std::filesystem::file_status Status = Entry.status(ErrorCode);
+		if (ErrorCode)
+		{
+			return Result;
+		}
 
-bool
-SetFileReadOnly(const FPath& path, bool bReadOnly)
-{
-	UNSYNC_WARNING(L"SetFileReadOnly() is not implemented");
-	return false;
-}
+		Result.bDirectory = Entry.is_directory();
+		Result.Size		  = Result.bDirectory ? 0 : Entry.file_size();
+		Result.Mtime	  = ToWindowsFileTime(Entry.last_write_time());
+		Result.bReadOnly  = (Status.permissions() & std::filesystem::perms::owner_write) == std::filesystem::perms::none;
+		Result.bValid	  = true;
+	}
 
-bool
-SetFileMtime(const FPath& path, uint64 mtime, bool bAllowInDryRun)
-{
-	UNSYNC_WARNING(L"SetFileMtime() is not implemented");
-	return false;
+	return Result;
 }
 
 uint64
 ToWindowsFileTime(const std::filesystem::file_time_type& FileTime)
 {
-	const uint64 NANOS_PER_TICK	  = 100ull;
-	const uint64 TICKS_PER_SECOND = 1'000'000'000ull / NANOS_PER_TICK;	// each tick is 100ns
-
-	// Windows epoch : 1601-01-01T00:00:00Z
-	// Unix epoch    : 1970-01-01T00:00:00Z
-	const uint64 SECONDS_BETWEEN_WINDOWS_AND_UNIX = 11'644'473'600ull;
-
-	//auto SysTime = std::chrono::file_clock::to_sys(FileTime);
-	//std::chrono::duration FullDuration = SysTime.time_since_epoch();
 	std::chrono::duration FullDuration = FileTime.time_since_epoch();
 
 	uint64 FullSeconds = std::chrono::floor<std::chrono::seconds>(FullDuration).count();
@@ -853,9 +813,9 @@ ToWindowsFileTime(const std::filesystem::file_time_type& FileTime)
 	auto SubsecondDuration = FullDuration - std::chrono::seconds(FullSeconds);
 	auto SubsecondNanos	   = std::chrono::duration_cast<std::chrono::nanoseconds>(SubsecondDuration).count();
 
-	uint64 ticks = (FullSeconds + SECONDS_BETWEEN_WINDOWS_AND_UNIX) * TICKS_PER_SECOND + (SubsecondNanos / NANOS_PER_TICK);
+	uint64 Ticks = (FullSeconds + SECONDS_BETWEEN_WINDOWS_AND_UNIX) * WINDOWS_TICKS_PER_SECOND + (SubsecondNanos / NANOS_PER_WINDOWS_TICK);
 
-	return ticks;
+	return Ticks;
 }
 
 uint64
@@ -866,6 +826,46 @@ GetAvailableDiskSpace(const FPath& Path)
 
 #endif	// UNSYNC_PLATFORM_UNIX
 
+bool
+SetFileMtime(const FPath& Path, uint64 Mtime, bool bAllowInDryRun)
+{
+	UNSYNC_ASSERT(!GDryRun || bAllowInDryRun);
+
+	FPath ExtendedPath = MakeExtendedAbsolutePath(Path);
+
+	std::filesystem::file_time_type FileTime = FromWindowsFileTime(Mtime);
+
+	std::error_code ErrorCode;
+	std::filesystem::last_write_time(ExtendedPath, FileTime, ErrorCode);
+
+	return !ErrorCode;
+}
+
+bool
+SetFileReadOnly(const FPath& Path, bool bReadOnly)
+{
+	UNSYNC_ASSERT(!GDryRun);
+
+	FPath ExtendedPath = MakeExtendedAbsolutePath(Path);
+
+	std::error_code ErrorCode;
+
+	if (bReadOnly)
+	{
+		std::filesystem::permissions(
+			ExtendedPath,
+			std::filesystem::perms::owner_write | std::filesystem::perms::group_write | std::filesystem::perms::others_write,
+			std::filesystem::perm_options::remove,
+			ErrorCode);
+	}
+	else
+	{
+		std::filesystem::permissions(ExtendedPath, std::filesystem::perms::owner_write, std::filesystem::perm_options::add, ErrorCode);
+	}
+
+	return !ErrorCode;
+}
+
 FBuffer
 ReadFileToBuffer(const FPath& Filename)
 {
@@ -1292,4 +1292,102 @@ FIOBuffer::operator=(FIOBuffer&& Rhs)
 	return *this;
 }
 
+void
+TestFileTime()
+{
+	UNSYNC_LOG(L"TestFileTime()");
+	UNSYNC_LOG_INDENT;
+
+	// 20231024004826Z - 2023 October 24 12:48:26
+	// unix 1698108506
+	// windows 133425821060000000
+	const uint64 BaseExpectedWindowsTime = 133425821060000000ull;
+
+	// Check basic conversion functionality at maximum
+	{
+		UNSYNC_LOG(L"File time precision estimate:");
+		UNSYNC_LOG_INDENT;
+
+		uint64 ExpectedWindowsTime = BaseExpectedWindowsTime + 9999999;
+
+		std::filesystem::file_time_type FileTime = FromWindowsFileTime(ExpectedWindowsTime);
+
+		uint64 RoundTripWindowsTime = ToWindowsFileTime(FileTime);
+		uint64 NativeCount			= FileTime.time_since_epoch().count();
+
+		uint64 Delta = ExpectedWindowsTime > RoundTripWindowsTime ? ExpectedWindowsTime - RoundTripWindowsTime
+																  : RoundTripWindowsTime - ExpectedWindowsTime;
+
+		UNSYNC_LOG(L"ExpectedWindowsTime  = %llu", llu(ExpectedWindowsTime));
+		UNSYNC_LOG(L"RoundTripWindowsTime = %llu", llu(RoundTripWindowsTime));
+		UNSYNC_LOG(L"NativeCount = %llu, Delta = %llu", llu(NativeCount), llu(Delta));
+	}
+
+	// Check basic conversion functionality at 1 second precision
+	{
+		uint64 ExpectedWindowsTime = BaseExpectedWindowsTime;
+
+		std::filesystem::file_time_type FileTime = FromWindowsFileTime(ExpectedWindowsTime);
+
+		uint64 RoundTripWindowsTime = ToWindowsFileTime(FileTime);
+		uint64 NativeCount			= FileTime.time_since_epoch().count();
+
+		UNSYNC_ASSERTF(RoundTripWindowsTime == ExpectedWindowsTime,
+					   L"RoundTripWindowsTime is %llu, but expected to be %llu. Native count: %llu",
+					   llu(RoundTripWindowsTime),
+					   llu(ExpectedWindowsTime),
+					   llu(NativeCount));
+	}
+}
+
+void
+TestFileAttrib()
+{
+	UNSYNC_LOG(L"TestFileAttrib()");
+	UNSYNC_LOG_INDENT;
+
+	FPath TempDirPath = std::filesystem::temp_directory_path() / "unsync_test";
+	CreateDirectories(TempDirPath);
+
+	const bool bDirectoryExists = PathExists(TempDirPath) && IsDirectory(TempDirPath);
+	UNSYNC_ASSERT(bDirectoryExists);
+
+	const FPath TestFilename = TempDirPath / "attrib.txt";
+	UNSYNC_LOG(L"Test file name: %ls", TestFilename.wstring().c_str());
+
+	if (PathExists(TestFilename))
+	{
+		SetFileReadOnly(TestFilename, false);
+	}
+
+	const bool	bFileWritten = WriteBufferToFile(TestFilename, "unsync test file");
+	UNSYNC_ASSERT(bFileWritten);
+
+	const uint64 ExpectedFileTime = 133425821060000000ull;
+
+	const bool bMtimeSet = SetFileMtime(TestFilename, ExpectedFileTime);
+	UNSYNC_ASSERT(bMtimeSet);
+
+	const FFileAttributes FileAttrib = GetFileAttrib(TestFilename);
+	UNSYNC_ASSERT(!FileAttrib.bReadOnly);
+
+	UNSYNC_ASSERT(FileAttrib.Mtime == ExpectedFileTime);
+
+	const bool bReadOnlySet = SetFileReadOnly(TestFilename, true);
+	UNSYNC_ASSERT(bReadOnlySet);
+
+	const FFileAttributes FileAttribReadOnly = GetFileAttrib(TestFilename);
+	UNSYNC_ASSERT(FileAttribReadOnly.bReadOnly);
+
+	const bool bReadOnlyReset = SetFileReadOnly(TestFilename, false);
+	UNSYNC_ASSERT(bReadOnlyReset);
+
+	const FFileAttributes FileAttribNonReadOnly = GetFileAttrib(TestFilename);
+	UNSYNC_ASSERT(!FileAttribNonReadOnly.bReadOnly);
+
+	std::error_code ErrorCode;
+	const bool		bFileDeleted = FileRemove(TestFilename, ErrorCode);
+	UNSYNC_ASSERT(bFileDeleted);
+}
+
 }  // namespace unsync
diff --git a/Engine/Source/Programs/Unsync/Private/UnsyncFile.h b/Engine/Source/Programs/Unsync/Private/UnsyncFile.h
--- a/Engine/Source/Programs/Unsync/Private/UnsyncFile.h
+++ b/Engine/Source/Programs/Unsync/Private/UnsyncFile.h
@@ -464,6 +464,7 @@ uint64 GetAvailableDiskSpace(const FPath& Path);
 std::filesystem::recursive_directory_iterator RecursiveDirectoryScan(const FPath& Path);
 
 uint64 ToWindowsFileTime(const std::filesystem::file_time_type& T);
+std::filesystem::file_time_type FromWindowsFileTime(uint64 Ticks);
 
 struct FSyncFilter;
 FFileAttributeCache CreateFileAttributeCache(const FPath& Root, const FSyncFilter* SyncFilter = nullptr);
diff --git a/Engine/Source/Programs/Unsync/Private/UnsyncMain.cpp b/Engine/Source/Programs/Unsync/Private/UnsyncMain.cpp
--- a/Engine/Source/Programs/Unsync/Private/UnsyncMain.cpp
+++ b/Engine/Source/Programs/Unsync/Private/UnsyncMain.cpp
@@ -46,7 +46,8 @@ InnerMain(int Argc, char** Argv)
 		"Copyright Epic Games, Inc. All Rights Reserved.\n";
 
 	CLI::App Cli(AppDescription, "unsync");
-	Cli.allow_windows_style_options(true);
+
+	Cli.allow_windows_style_options(false); // never allow /flag syntax, only --flag
 	Cli.set_version_flag("--version", GetVersionString());
 
 	std::vector<CLI::App*> SubCommands;
@@ -857,13 +858,13 @@ InnerMain(int Argc, char** Argv)
 		}
 
 		FCmdLoginOptions LoginOptions;
-		LoginOptions.Remote		   = RemoteDesc;
-		LoginOptions.bInteractive  = bInteractive;
-		LoginOptions.bDecode	   = bDecode;
-		LoginOptions.bPrint		   = bPrint;
+		LoginOptions.Remote			  = RemoteDesc;
+		LoginOptions.bInteractive	  = bInteractive;
+		LoginOptions.bDecode		  = bDecode;
+		LoginOptions.bPrint			  = bPrint;
 		LoginOptions.bPrintHttpHeader = bPrintHttpHeader;
-		LoginOptions.bForceRefresh = bForceRefreshAuth;
-		LoginOptions.bQuick		   = bQuickLogin;
+		LoginOptions.bForceRefresh	  = bForceRefreshAuth;
+		LoginOptions.bQuick			  = bQuickLogin;
 		return CmdLogin(LoginOptions);
 	}
 	else if (Cli.got_subcommand(SubMount))
diff --git a/Engine/Source/Programs/Unsync/Private/UnsyncTest.cpp b/Engine/Source/Programs/Unsync/Private/UnsyncTest.cpp
--- a/Engine/Source/Programs/Unsync/Private/UnsyncTest.cpp
+++ b/Engine/Source/Programs/Unsync/Private/UnsyncTest.cpp
@@ -915,6 +915,18 @@ RunTests(const std::string& Preset)
 		extern void TestMiniCb();
 		TestMiniCb();
 	}
+
+	if (Preset == "filetime" || Preset == "all")
+	{
+		extern void TestFileTime();
+		TestFileTime();
+	}
+
+	if (Preset == "fileattrib" || Preset == "all")
+	{
+		extern void TestFileAttrib();
+		TestFileAttrib();
+	}
 }
 
 }  // namespace unsync
