unsync - Store read-only file status in the manifest and apply it after sync, tweak background task balancing and other minor optimizations

* Use explicit file attribute cache when possible, to avoid redundant filesystem ops in some cases
* Skip generating stable manifest signature during normal sync
* Fix few redundant memory allocations
* Add few log events around potentially expensive ops

diff --git a/Engine/Source/Programs/Unsync/Private/UnsyncCmdPush.cpp b/Engine/Source/Programs/Unsync/Private/UnsyncCmdPush.cpp
--- a/Engine/Source/Programs/Unsync/Private/UnsyncCmdPush.cpp
+++ b/Engine/Source/Programs/Unsync/Private/UnsyncCmdPush.cpp
@@ -51,7 +51,7 @@ CmdPush(const FCmdPushOptions& Options)
 	}
 
 	FDirectoryManifestInfo ManifestInfo			= GetManifestInfo(Manifest);
-	FHash160			   ManifestSignature	= ToHash160(ManifestInfo.Signature);
+	FHash160			   ManifestSignature	= ToHash160(ManifestInfo.StableSignature);
 	std::string			   ManifestSignatureStr = BytesToHexString(ManifestSignature.Data, sizeof(ManifestSignature.Data));
 
 	LogManifestInfo(ELogLevel::Debug, ManifestInfo);
diff --git a/Engine/Source/Programs/Unsync/Private/UnsyncCore.cpp b/Engine/Source/Programs/Unsync/Private/UnsyncCore.cpp
--- a/Engine/Source/Programs/Unsync/Private/UnsyncCore.cpp
+++ b/Engine/Source/Programs/Unsync/Private/UnsyncCore.cpp
@@ -27,7 +27,7 @@ UNSYNC_THIRD_PARTY_INCLUDES_START
 #include <md5-sse2.h>
 UNSYNC_THIRD_PARTY_INCLUDES_END
 
-#define UNSYNC_VERSION_STR "1.0.61-dev"
+#define UNSYNC_VERSION_STR "1.0.61-dev2"
 
 namespace unsync {
 
@@ -1828,6 +1828,7 @@ CreateDirectoryManifest(const FPath& Root, uint32 BlockSize, FAlgorithmOptions A
 		FileManifest.Size		 = Dir.file_size();
 		FileManifest.CurrentPath = Dir.path();
 		FileManifest.BlockSize	 = BlockSize;
+		FileManifest.bReadOnly	 = IsReadOnly(Dir.status().permissions());
 
 		{
 			std::lock_guard<std::mutex> LockGuard(ResultMutex);
@@ -2089,7 +2090,10 @@ LoadOrCreateDirectoryManifest(FDirectoryManifest& Result, const FPath& Root, uin
 		// Otherwise just do a quick manifest generation, without file blocks.
 		uint32 NewManifestBlockSize = GDryRun ? BlockSize : 0;
 
+		UNSYNC_VERBOSE(L"Creating lightweight manifest for '%ls'", Root.wstring().c_str());
 		NewDirectoryManifest = CreateDirectoryManifest(Root, NewManifestBlockSize, Algorithm);
+
+		UNSYNC_VERBOSE(L"Comparing manifests");
 		for (const auto& OldManifestIt : OldDirectoryManifest.Files)
 		{
 			auto NewManifestIt = NewDirectoryManifest.Files.find(OldManifestIt.first);
@@ -2125,6 +2129,7 @@ LoadOrCreateDirectoryManifest(FDirectoryManifest& Result, const FPath& Root, uin
 
 		if (BlockSize)
 		{
+			UNSYNC_VERBOSE(L"Updating manifest blocks");
 			UpdateDirectoryManifestBlocks(NewDirectoryManifest, Root, BlockSize, Algorithm);
 		}
 	}
@@ -2179,10 +2184,11 @@ ComputeManifestStableSignature(const FDirectoryManifest& Manifest)
 
 	std::sort(SortedFiles.begin(), SortedFiles.end());
 
+	std::string FileNameUtf8;
 	for (const std::wstring& FileName : SortedFiles)
 	{
 		// Canonical unsync file paths are utf8 with unix-style separator `/`
-		std::string FileNameUtf8 = ConvertWideToUtf8(FileName);
+		ConvertWideToUtf8(FileName, FileNameUtf8);
 		std::replace(FileNameUtf8.begin(), FileNameUtf8.end(), '\\', '/');
 
 		const FFileManifest& FileManifest = Manifest.Files.at(FileName);
@@ -2506,17 +2512,17 @@ SyncFile(const FNeedList&		   NeedList,
 				}
 			}
 
-			std::error_code Ec = {};
-			FileRename(TempTargetFilePath, TargetFilePath, Ec);
+			std::error_code ErrorCode = {};
+			FileRename(TempTargetFilePath, TargetFilePath, ErrorCode);
 
-			if (Ec.value() == 0)
+			if (ErrorCode.value() == 0)
 			{
 				Result.Status = EFileSyncStatus::Ok;
 			}
 			else
 			{
 				Result.Status		   = EFileSyncStatus::ErrorFinalRename;
-				Result.SystemErrorCode = Ec;
+				Result.SystemErrorCode = ErrorCode;
 			}
 		}
 
@@ -3320,6 +3326,8 @@ SyncDirectory(const FSyncDirectoryOptions& SyncOptions)
 
 	const bool bCaseSensitiveTargetFileSystem = IsCaseSensitiveFileSystem(TargetTempPath);
 
+	FTimingLogger ManifestLoadTimingLogger("Manifest load time");
+
 	if (SyncOptions.SourceType == ESyncSourceType::ServerWithManifestHash)
 	{
 		if (!ProxyPool.IsValid())
@@ -3375,7 +3383,7 @@ SyncDirectory(const FSyncDirectoryOptions& SyncOptions)
 	{
 		UNSYNC_VERBOSE(L"Loaded manifest properties:");
 		UNSYNC_LOG_INDENT;
-		FDirectoryManifestInfo ManifestInfo = GetManifestInfo(SourceDirectoryManifest);
+		FDirectoryManifestInfo ManifestInfo = GetManifestInfo(SourceDirectoryManifest, false /*bGenerateSignature*/);
 		LogManifestInfo(ELogLevel::Debug, ManifestInfo);
 		if (ProxyPool.RemoteDesc.Protocol == EProtocolFlavor::Jupiter && ManifestInfo.NumMacroBlocks == 0)
 		{
@@ -3384,9 +3392,15 @@ SyncDirectory(const FSyncDirectoryOptions& SyncOptions)
 		}
 	}
 
+	ManifestLoadTimingLogger.Finish();
+
 	// Propagate algorithm selection from source
 	FAlgorithmOptions  Algorithm			   = SourceDirectoryManifest.Options;
+
+	FTimingLogger TargetManifestTimingLogger("Target directory manifest generation time");
+	UNSYNC_VERBOSE(L"Creating manifest for directory '%ls'", TargetPath.wstring().c_str());
 	FDirectoryManifest TargetDirectoryManifest = CreateDirectoryManifest(TargetPath, 0, Algorithm);
+	TargetManifestTimingLogger.Finish();
 
 	if (!bCaseSensitiveTargetFileSystem)
 	{
@@ -3525,7 +3539,7 @@ SyncDirectory(const FSyncDirectoryOptions& SyncOptions)
 
 		if (bSourceManifestOk)
 		{
-			FFileAttributes BaseFileAttrib = GetFileAttrib(BaseFilePath, &BaseAttribCache);
+			FFileAttributes BaseFileAttrib = GetCachedFileAttrib(BaseFilePath, BaseAttribCache);
 
 			if (!BaseFileAttrib.bValid)
 			{
@@ -3770,6 +3784,11 @@ SyncDirectory(const FSyncDirectoryOptions& SyncOptions)
 	UNSYNC_VERBOSE(L"Copying files");
 
 	{
+		// Throttle background tasks by trying to keep them to some sensible memory budget. Best effort only, not a hard limit.
+		static constexpr uint64 BackgroundTaskMemoryBudget	= 2_GB;
+		static constexpr uint64 TargetTotalSizePerTaskBatch = BackgroundTaskMemoryBudget;
+		static constexpr uint64 MaxFilesPerTaskBatch		= 1000;
+
 		struct FBackgroundTaskResult
 		{
 			FPath			TargetFilePath;
@@ -3796,7 +3815,25 @@ SyncDirectory(const FSyncDirectoryOptions& SyncOptions)
 			FFileSyncTaskBatch CurrentBatch;
 			for (const FFileSyncTask& FileTask : AllFileTasks)
 			{
-				if (CurrentBatch.FileTasks.size() && CurrentBatch.NeedBytesFromSource + FileTask.NeedBytesFromSource > MaxBatchDownloadSize)
+				bool bShouldBreakBatch = false;
+
+				if (!CurrentBatch.FileTasks.empty())
+				{
+					if (CurrentBatch.NeedBytesFromSource + FileTask.NeedBytesFromSource > MaxBatchDownloadSize)
+					{
+						bShouldBreakBatch = true;
+					}
+					else if (CurrentBatch.FileTasks.size() >= MaxFilesPerTaskBatch)
+					{
+						bShouldBreakBatch = true;
+					}
+					else if (CurrentBatch.TotalSizeBytes >= TargetTotalSizePerTaskBatch)
+					{
+						bShouldBreakBatch = true;
+					}
+				}
+
+				if (bShouldBreakBatch)
 				{
 					SyncTaskList.push_back(std::move(CurrentBatch));
 					CurrentBatch = {};
@@ -3869,6 +3906,10 @@ SyncDirectory(const FSyncDirectoryOptions& SyncOptions)
 				{
 					BaseFile = nullptr;
 					SetFileMtime(Item.TargetFilePath, Item.SourceManifest->Mtime);
+					if (Item.SourceManifest->bReadOnly)
+					{
+						SetFileReadOnly(Item.TargetFilePath, true);
+					}
 				}
 
 				if (bBackground)
@@ -3911,8 +3952,6 @@ SyncDirectory(const FSyncDirectoryOptions& SyncOptions)
 		FTaskGroup BackgroundTaskGroup;
 		FTaskGroup ForegroundTaskGroup;
 
-		// Throttle background tasks by trying to keep them to some sensible memory budget. Best effort only, not a hard limit.
-		const uint64		BackgroundTaskMemoryBudget = 2_GB;
 		std::atomic<uint64> BackgroundTaskMemory	   = {};
 		std::atomic<uint64> RemainingSourceBytes	   = EstimatedNeedBytesFromSource;
 
@@ -4585,7 +4624,7 @@ ToString(EFileSyncStatus Status)
 }
 
 FDirectoryManifestInfo
-GetManifestInfo(const FDirectoryManifest& Manifest)
+GetManifestInfo(const FDirectoryManifest& Manifest, bool bGenerateSignature)
 {
 	FDirectoryManifestInfo Result = {};
 
@@ -4621,8 +4660,10 @@ GetManifestInfo(const FDirectoryManifest& Manifest)
 	Result.NumMacroBlocks = UniqueMacroBlockSet.size();
 	Result.NumFiles		  = Manifest.Files.size();
 	Result.Algorithm	  = Manifest.Options;
-	Result.SerializedHash = ComputeSerializedManifestHash(Manifest);
-	Result.Signature	  = ComputeManifestStableSignature(Manifest);
+	if (bGenerateSignature)
+	{
+		Result.StableSignature = ComputeManifestStableSignature(Manifest);
+	}
 
 	return Result;
 }
@@ -4630,10 +4671,15 @@ GetManifestInfo(const FDirectoryManifest& Manifest)
 void
 LogManifestInfo(ELogLevel LogLevel, const FDirectoryManifestInfo& Info)
 {
-	FHash160	ManifestSignature = ToHash160(Info.Signature);
-	std::string SignatureHexStr	  = HashToHexString(ManifestSignature);
+	FHash160	ManifestSignature = ToHash160(Info.StableSignature);
+	const FHash160 EmptySignature = {};
+
+	if (ManifestSignature != EmptySignature)
+	{
+		std::string SignatureHexStr = HashToHexString(ManifestSignature);
+		LogPrintf(LogLevel, L"Manifest signature: %hs\n", SignatureHexStr.c_str());
+	}
 
-	LogPrintf(LogLevel, L"Manifest signature: %hs\n", SignatureHexStr.c_str());
 	LogPrintf(LogLevel, L"Chunking mode: %hs\n", ToString(Info.Algorithm.ChunkingAlgorithmId));
 	LogPrintf(LogLevel, L"Weak hash: %hs\n", ToString(Info.Algorithm.WeakHashAlgorithmId));
 	LogPrintf(LogLevel, L"Strong hash: %hs\n", ToString(Info.Algorithm.StrongHashAlgorithmId));
diff --git a/Engine/Source/Programs/Unsync/Private/UnsyncCore.h b/Engine/Source/Programs/Unsync/Private/UnsyncCore.h
--- a/Engine/Source/Programs/Unsync/Private/UnsyncCore.h
+++ b/Engine/Source/Programs/Unsync/Private/UnsyncCore.h
@@ -128,6 +128,7 @@ struct FFileManifest
 	FGenericBlockArray Blocks;
 	FGenericBlockArray MacroBlocks;
 	FPath			   CurrentPath;
+	bool			   bReadOnly = false;
 };
 
 struct FAlgorithmOptionsV5
@@ -179,18 +180,17 @@ struct FDirectoryManifest
 
 struct FDirectoryManifestInfo
 {
-	uint64	 TotalSize		= 0;
-	uint64	 UniqueSize		= 0;
-	uint64	 NumBlocks		= 0;
-	uint64	 NumMacroBlocks = 0;
-	uint64	 NumFiles		= 0;
-	FHash256 SerializedHash = {};
-	FHash256 Signature		= {};
+	uint64	 TotalSize		 = 0;
+	uint64	 UniqueSize		 = 0;
+	uint64	 NumBlocks		 = 0;
+	uint64	 NumMacroBlocks	 = 0;
+	uint64	 NumFiles		 = 0;
+	FHash256 StableSignature = {};
 
 	FAlgorithmOptions Algorithm = {};
 };
 
-FDirectoryManifestInfo GetManifestInfo(const FDirectoryManifest& Manifest);
+FDirectoryManifestInfo GetManifestInfo(const FDirectoryManifest& Manifest, bool bGenerateSignature = true);
 void				   LogManifestInfo(ELogLevel LogLevel, const FDirectoryManifestInfo& Info);
 void				   LogManifestFiles(ELogLevel LogLevel, const FDirectoryManifestInfo& Info);
 
diff --git a/Engine/Source/Programs/Unsync/Private/UnsyncFile.cpp b/Engine/Source/Programs/Unsync/Private/UnsyncFile.cpp
--- a/Engine/Source/Programs/Unsync/Private/UnsyncFile.cpp
+++ b/Engine/Source/Programs/Unsync/Private/UnsyncFile.cpp
@@ -4,6 +4,7 @@
 #include "UnsyncCore.h"
 #include "UnsyncMemory.h"
 #include "UnsyncThread.h"
+#include "UnsyncThread.h"
 
 #include <mutex>
 
@@ -125,6 +126,21 @@ GetRelativePath(const FPath& Path, const FPath& Base)
 	return std::filesystem::relative(Path, Base);
 }
 
+FFileAttributes GetCachedFileAttrib(const FPath& Path, FFileAttributeCache& AttribCache)
+{
+	FFileAttributes Result;
+
+	FPath ExtendedPath = MakeExtendedAbsolutePath(Path);
+
+	auto It = AttribCache.Map.find(ExtendedPath);
+	if (It != AttribCache.Map.end())
+	{
+		Result = It->second;
+	}
+
+	return Result;
+}
+
 #if UNSYNC_PLATFORM_WINDOWS
 inline uint64
 MakeU64(FILETIME Ft)
@@ -796,7 +812,7 @@ GetFileAttrib(const FPath& Path, FFileAttributeCache* AttribCache)
 		Result.bDirectory = Entry.is_directory();
 		Result.Size		  = Result.bDirectory ? 0 : Entry.file_size();
 		Result.Mtime	  = ToWindowsFileTime(Entry.last_write_time());
-		Result.bReadOnly  = (Status.permissions() & std::filesystem::perms::owner_write) == std::filesystem::perms::none;
+		Result.bReadOnly  = IsReadOnly(Status.permissions());
 		Result.bValid	  = true;
 	}
 
@@ -1117,6 +1133,7 @@ CreateFileAttributeCache(const FPath& Root, const FSyncFilter* SyncFilter)
 		Attr.Mtime	= ToWindowsFileTime(Dir.last_write_time());
 		Attr.Size	= Dir.file_size();
 		Attr.bValid = true;
+		Attr.bReadOnly = IsReadOnly(Dir.status().permissions());
 
 		Result.Map[Dir.path().native()] = Attr;
 
diff --git a/Engine/Source/Programs/Unsync/Private/UnsyncFile.h b/Engine/Source/Programs/Unsync/Private/UnsyncFile.h
--- a/Engine/Source/Programs/Unsync/Private/UnsyncFile.h
+++ b/Engine/Source/Programs/Unsync/Private/UnsyncFile.h
@@ -444,7 +444,15 @@ struct FFileAttributeCache
 	const bool Exists(const FPath& Path) const;
 };
 
+inline bool
+IsReadOnly(std::filesystem::perms Perms)
+{
+	return (Perms & std::filesystem::perms::owner_write) == std::filesystem::perms::none;
+}
+
 FFileAttributes GetFileAttrib(const FPath& Path, FFileAttributeCache* AttribCache = nullptr);
+FFileAttributes GetCachedFileAttrib(const FPath& Path, FFileAttributeCache& AttribCache);
+
 bool			SetFileMtime(const FPath& Path, uint64 Mtime, bool bAllowInDryRun = false);
 bool			SetFileReadOnly(const FPath& Path, bool ReadOnly);
 bool			IsDirectory(const FPath& Path);
diff --git a/Engine/Source/Programs/Unsync/Private/UnsyncLog.cpp b/Engine/Source/Programs/Unsync/Private/UnsyncLog.cpp
--- a/Engine/Source/Programs/Unsync/Private/UnsyncLog.cpp
+++ b/Engine/Source/Programs/Unsync/Private/UnsyncLog.cpp
@@ -252,12 +252,12 @@ LogPrintf(ELogLevel Level, const wchar_t* Str, ...)
 	if (Level == ELogLevel::Error)
 	{
 		Prefix = L"ERROR: ";
-		bShouldOutputThreadIndex = true;
+		bShouldOutputThreadIndex = ThreadIndex != 0;
 	}
 	else if (Level == ELogLevel::Warning)
 	{
 		Prefix = L"WARNING: ";
-		bShouldOutputThreadIndex = true;
+		bShouldOutputThreadIndex = ThreadIndex != 0;
 	}
 	else if (GLogIndent)
 	{
diff --git a/Engine/Source/Programs/Unsync/Private/UnsyncProtocol.h b/Engine/Source/Programs/Unsync/Private/UnsyncProtocol.h
--- a/Engine/Source/Programs/Unsync/Private/UnsyncProtocol.h
+++ b/Engine/Source/Programs/Unsync/Private/UnsyncProtocol.h
@@ -65,9 +65,10 @@ using FBlock256 = TBlock<FHash256>;
 
 using FGenericBlock = TBlock<FGenericHash>;
 
-static constexpr uint64 SERIALIZED_SECTION_ID_TERMINATOR	  = 0;
-static constexpr uint64 SERIALIZED_SECTION_ID_METADATA_STRING = 0xC6BD6CDCEEF79533ull;
-static constexpr uint64 SERIALIZED_SECTION_ID_MACRO_BLOCK	  = 0x8390AEBB745E08BCull;
+static constexpr uint64 SERIALIZED_SECTION_ID_TERMINATOR		  = 0;
+static constexpr uint64 SERIALIZED_SECTION_ID_METADATA_STRING	  = 0xC6BD6CDCEEF79533ull;
+static constexpr uint64 SERIALIZED_SECTION_ID_MACRO_BLOCK		  = 0x8390AEBB745E08BCull;
+static constexpr uint64 SERIALIZED_SECTION_ID_FILE_READ_ONLY_MASK = 0x851F32ED3615F0ADull;
 
 struct FSerializedSectionHeader
 {
@@ -92,6 +93,12 @@ struct FMacroBlockSection
 	static constexpr uint64 VERSION = 2;
 };
 
+struct FFileReadOnlyMaskSection
+{
+	static constexpr uint64 MAGIC	= SERIALIZED_SECTION_ID_FILE_READ_ONLY_MASK;
+	static constexpr uint64 VERSION = 1;
+};
+
 struct FBlockFileHeader
 {
 	static constexpr uint64 MAGIC	= 0x1DCB86A5BDBA27CFull;
diff --git a/Engine/Source/Programs/Unsync/Private/UnsyncSerialization.cpp b/Engine/Source/Programs/Unsync/Private/UnsyncSerialization.cpp
--- a/Engine/Source/Programs/Unsync/Private/UnsyncSerialization.cpp
+++ b/Engine/Source/Programs/Unsync/Private/UnsyncSerialization.cpp
@@ -282,6 +282,27 @@ LoadMacroBlocks(FIOReaderStream& Reader, uint64 Version)
 	return Result;
 }
 
+static FBuffer
+LoadFileReadOnlyMask(FIOReaderStream& Reader, FSerializedSectionHeader Header)
+{
+	FBuffer Result;
+
+	if (Header.Version != FFileReadOnlyMaskSection::VERSION)
+	{
+		UNSYNC_ERROR(L"Found read-only file mask section version %llu, but only version %lly is supported.",
+					 llu(Header.Version),
+					 llu(FFileReadOnlyMaskSection::VERSION));
+		return Result;
+	}
+
+	Result.Resize(Header.Size);
+	uint8* ResultData = Result.Data();
+
+	Reader.Read(ResultData, Header.Size);
+
+	return Result;
+}
+
 static FBuffer
 SaveMacroBlocks(const FDirectoryManifest& Manifest)
 {
@@ -316,6 +337,29 @@ SaveMacroBlocks(const FDirectoryManifest& Manifest)
 	return Result;
 }
 
+static FBuffer
+SaveFileReadOnlyMask(const FDirectoryManifest& Manifest)
+{
+	FBuffer			 Result;
+	FVectorStreamOut Writer(Result);
+
+	const uint64 NumFiles		  = Manifest.Files.size();
+	const uint64 NumBitArrayBytes = DivUp(NumFiles, 8);
+
+	Result.Resize(NumBitArrayBytes);
+	memset(Result.Data(), 0, Result.Size());
+
+	uint64 FileIndex = 0;
+	for (const auto& FileIt : Manifest.Files)
+	{
+		const FFileManifest& FileManifest = FileIt.second;
+		BitArraySet(Result.Data(), FileIndex, FileManifest.bReadOnly);
+		++FileIndex;
+	}
+
+	return Result;
+}
+
 bool  // TODO: return a TResult
 LoadDirectoryManifest(FDirectoryManifest& OutManifest, const FPath& Root, FIOReaderStream& Stream)
 {
@@ -365,6 +409,8 @@ LoadDirectoryManifest(FDirectoryManifest& OutManifest, const FPath& Root, FIORea
 
 	std::unordered_map<FHash256, FGenericBlockArray> MacroBlocks;
 
+	FBuffer FileReadOnlyMask;
+
 	if (Version >= FDirectoryManifest::EVersions::V7_OptionalSections)
 	{
 		// Load optional sections until terminator is encountered
@@ -399,6 +445,11 @@ LoadDirectoryManifest(FDirectoryManifest& OutManifest, const FPath& Root, FIORea
 							MacroBlocks = LoadMacroBlocks(Stream, SectionHeader.Version);
 							break;
 						}
+					case SERIALIZED_SECTION_ID_FILE_READ_ONLY_MASK:
+						{
+							FileReadOnlyMask = LoadFileReadOnlyMask(Stream, SectionHeader);
+							break;
+						}
 					case SERIALIZED_SECTION_ID_TERMINATOR:
 						bDoneLoadingOptionalSections = true;
 						break;
@@ -427,9 +478,12 @@ LoadDirectoryManifest(FDirectoryManifest& OutManifest, const FPath& Root, FIORea
 	if (AlgorithmCompatibility == EAlgorithmCompatibilityResult::Ok)
 	{
 		Serialize(Stream, NumFiles);
-		for (uint64 I = 0; I < NumFiles; ++I)
+
+		const bool bReadOnlyMaskValid = NumFiles <= FileReadOnlyMask.Size() * 8;
+
+		std::string FilenameUtf8; // shared buffer to avoid some of the reallocations
+		for (uint64 FileIndex = 0; FileIndex < NumFiles; ++FileIndex)
 		{
-			std::string FilenameUtf8;
 			Serialize(Stream, FilenameUtf8);
 
 			FFileManifest FileManifest;
@@ -479,6 +533,13 @@ LoadDirectoryManifest(FDirectoryManifest& OutManifest, const FPath& Root, FIORea
 			}
 
 			FileManifest.CurrentPath	= Root / Filename;
+
+			if (bReadOnlyMaskValid)
+			{
+				FileManifest.bReadOnly = BitArrayGet(FileReadOnlyMask.Data(), FileIndex);
+			}
+
+			// Store output
 			OutManifest.Files[Filename] = std::move(FileManifest);
 		}
 	}
@@ -555,6 +616,13 @@ WriteSection(FVectorStreamOut& Stream, const FMetadataStringSection& Section)
 	WriteSection(Stream, FMetadataStringSection::MAGIC, FMetadataStringSection::VERSION, StreamBuffer.View());
 }
 
+template<typename SectionHeaderType>
+static void
+WriteSection(FVectorStreamOut& Stream, FBufferView SectionData)
+{
+	WriteSection(Stream, SectionHeaderType::MAGIC, SectionHeaderType::VERSION, SectionData);
+}
+
 bool
 SaveDirectoryManifest(const FDirectoryManifest& Manifest, FVectorStreamOut& Stream)
 {
@@ -592,8 +660,13 @@ SaveDirectoryManifest(const FDirectoryManifest& Manifest, FVectorStreamOut& Stre
 
 	{
 		// TODO: only save macro block section if it's valid
-		FBuffer SerializedMacroBlocks = SaveMacroBlocks(Manifest);
-		WriteSection(Stream, FMacroBlockSection::MAGIC, FMacroBlockSection::VERSION, SerializedMacroBlocks.View());
+		FBuffer SectionBuffer = SaveMacroBlocks(Manifest);
+		WriteSection<FMacroBlockSection>(Stream, SectionBuffer.View());
+	}
+
+	{
+		FBuffer SectionBuffer = SaveFileReadOnlyMask(Manifest);
+		WriteSection<FFileReadOnlyMaskSection>(Stream, SectionBuffer.View());
 	}
 
 	// End with the terminator section (default-constructed);
diff --git a/Engine/Source/Programs/Unsync/Private/UnsyncUtil.cpp b/Engine/Source/Programs/Unsync/Private/UnsyncUtil.cpp
--- a/Engine/Source/Programs/Unsync/Private/UnsyncUtil.cpp
+++ b/Engine/Source/Programs/Unsync/Private/UnsyncUtil.cpp
@@ -63,9 +63,9 @@ BytesToHexString(const uint8* Data, uint64 Size)
 	return Result;
 }
 
-FTimingLogger::FTimingLogger(const char* InName, bool InEnabled) : Enabled(InEnabled), Name(InName)
+FTimingLogger::FTimingLogger(const char* InName, bool bInEnabled) : bEnabled(bInEnabled), Name(InName)
 {
-	if (Enabled)
+	if (bEnabled)
 	{
 		TimeBegin = TimePointNow();
 	}
@@ -73,7 +73,12 @@ FTimingLogger::FTimingLogger(const char* InName, bool InEnabled) : Enabled(InEna
 
 FTimingLogger::~FTimingLogger()
 {
-	if (Enabled)
+	Finish();
+}
+
+void FTimingLogger::Finish()
+{
+	if (bEnabled)
 	{
 		FTimePoint	  TimeEnd	   = TimePointNow();
 		FTimeDuration Duration	   = FTimeDuration(TimeEnd - TimeBegin);
@@ -93,6 +98,8 @@ FTimingLogger::~FTimingLogger()
 		}
 
 		LogFlush();
+
+		bEnabled = false;
 	}
 }
 
@@ -137,10 +144,11 @@ ConvertUtf8ToWide(std::string_view StringUtf8)
 	return Result;
 }
 
-std::string
-ConvertWideToUtf8(std::wstring_view StringWide)
+
+void
+ConvertWideToUtf8(std::wstring_view StringWide, std::string& Result)
 {
-	std::string Result;
+	Result.clear();
 
 	if (IsTrivialAsciiString(StringWide))
 	{
@@ -155,7 +163,13 @@ ConvertWideToUtf8(std::wstring_view StringWide)
 		std::wstring_convert<std::codecvt_utf8<wchar_t>> Cvt;
 		Result = Cvt.to_bytes(StringWide.data(), StringWide.data() + StringWide.length());
 	}
+}
 
+std::string
+ConvertWideToUtf8(std::wstring_view StringWide)
+{
+	std::string Result;
+	ConvertWideToUtf8(StringWide, Result);
 	return Result;
 }
 
diff --git a/Engine/Source/Programs/Unsync/Private/UnsyncUtil.h b/Engine/Source/Programs/Unsync/Private/UnsyncUtil.h
--- a/Engine/Source/Programs/Unsync/Private/UnsyncUtil.h
+++ b/Engine/Source/Programs/Unsync/Private/UnsyncUtil.h
@@ -82,11 +82,13 @@ struct FRange
 
 struct FTimingLogger
 {
-	bool		Enabled	  = false;
+	bool		bEnabled  = false;
 	FTimePoint	TimeBegin = FTimePoint{};
 	std::string Name;
+
 	FTimingLogger(const char* InName, bool InEnabled = true);
 	~FTimingLogger();
+	void Finish();
 };
 
 inline double
@@ -134,6 +136,7 @@ HashToHexString(const HashType& Hash)
 
 std::wstring ConvertUtf8ToWide(std::string_view StringUtf8);
 std::string	 ConvertWideToUtf8(std::wstring_view StringWide);
+void		 ConvertWideToUtf8(std::wstring_view StringWide, std::string& Result);
 
 std::wstring StringToLower(const std::wstring& Input);
 std::wstring StringToUpper(const std::wstring& Input);
@@ -298,6 +301,38 @@ CountLeadingZeros64(uint64 X)
 	return ((63 - XLog2) & Mask) | (64 & ~Mask);
 }
 
+template<typename StorageT>
+struct TBitArrayInfo
+{
+	static constexpr size_t ElemSizeInBits	= sizeof(StorageT) * 8;
+	const uint64			ElemIndex;
+	const StorageT			BitMask;
+	TBitArrayInfo(uint64 BitIndex) : ElemIndex(BitIndex / ElemSizeInBits), BitMask(StorageT(1) << (BitIndex % ElemSizeInBits)) {}
+};
+
+template<typename StorageT>
+inline bool
+BitArrayGet(const StorageT* Storage, uint64 BitIndex)
+{
+	TBitArrayInfo<StorageT> Info(BitIndex);
+	return (Storage[Info.ElemIndex] & Info.BitMask) != 0;
+}
+
+template<typename StorageT>
+inline void
+BitArraySet(StorageT* Storage, uint64 BitIndex, bool bValue)
+{
+	TBitArrayInfo<StorageT> Info(BitIndex);
+	if (bValue)
+	{
+		Storage[Info.ElemIndex] |= Info.BitMask;
+	}
+	else
+	{
+		Storage[Info.ElemIndex] &= ~Info.BitMask;
+	}
+}
+
 FPath NormalizeFilenameUtf8(const std::string& InFilename);
 FPath GetAbsoluteNormalPath(const FPath& InPath);
 
