unsync - Always use simple lexical implementation of GetRelativePath

diff --git a/Engine/Source/Programs/Unsync/Private/UnsyncFile.cpp b/Engine/Source/Programs/Unsync/Private/UnsyncFile.cpp
--- a/Engine/Source/Programs/Unsync/Private/UnsyncFile.cpp
+++ b/Engine/Source/Programs/Unsync/Private/UnsyncFile.cpp
@@ -113,7 +113,10 @@ GetRelativePath(const FPath& Path, const FPath& Base)
 	// Try a trivial case first, without touching the filesystem
 	FPathStringView PathView = RemoveExtendedPathPrefix(Path);
 	FPathStringView BaseView = RemoveExtendedPathPrefix(Base);
-	if (PathView.starts_with(BaseView))
+
+	FPathStringView PathViewRemainder = PathView.substr(BaseView.length());
+
+	if (PathView.starts_with(BaseView) && PathViewRemainder.starts_with(FPath::preferred_separator))
 	{
 		FPathStringView RelativePath = PathView.substr(BaseView.length());
 		while (RelativePath.starts_with(FPath::preferred_separator))
@@ -123,7 +126,7 @@ GetRelativePath(const FPath& Path, const FPath& Base)
 		return FPath(RelativePath);
 	}
 
-	return std::filesystem::relative(Path, Base);
+	return {};
 }
 
 FFileAttributes GetCachedFileAttrib(const FPath& Path, FFileAttributeCache& AttribCache)
