unsync - Refactor `ComputeBlocks` to take a struct instead of loose parameters

diff --git a/Engine/Source/Programs/Unsync/Private/UnsyncCmdHash.cpp b/Engine/Source/Programs/Unsync/Private/UnsyncCmdHash.cpp
--- a/Engine/Source/Programs/Unsync/Private/UnsyncCmdHash.cpp
+++ b/Engine/Source/Programs/Unsync/Private/UnsyncCmdHash.cpp
@@ -12,7 +12,7 @@ CmdHash(const FCmdHashOptions& Options)
 {
 	if (unsync::IsDirectory(Options.Input))
 	{
-		UNSYNC_VERBOSE(L"'%ls' is a directory", Options.Input.wstring().c_str());
+		UNSYNC_VERBOSE(L"Generating manifest for directory '%ls'", Options.Input.wstring().c_str());
 
 		FPath InputRoot	   = Options.Input;
 		FPath ManifestRoot = InputRoot / ".unsync";
@@ -52,18 +52,26 @@ CmdHash(const FCmdHashOptions& Options)
 
 		if (!GDryRun)
 		{
+			UNSYNC_VERBOSE(L"Saving directory manifest '%ls'", DirectoryManifestPath.wstring().c_str());
+
 			SaveDirectoryManifest(DirectoryManifest, DirectoryManifestPath);
 		}
 	}
 	else
 	{
-		UNSYNC_VERBOSE(L"'%ls' is a file", Options.Input.wstring().c_str());
+		UNSYNC_VERBOSE(L"Generating manfiest for file '%ls'", Options.Input.wstring().c_str());
 
 		FNativeFile OverlappedFile(Options.Input);
 		if (OverlappedFile.IsValid())
 		{
 			UNSYNC_VERBOSE(L"Computing blocks for '%ls' (%.2f MB)", Options.Input.wstring().c_str(), SizeMb(OverlappedFile.GetSize()));
-			FGenericBlockArray GenericBlocks = ComputeBlocks(OverlappedFile, Options.BlockSize, Options.Algorithm);
+			FComputeBlocksParams ComputeBlocksParams;
+			ComputeBlocksParams.Algorithm	 = Options.Algorithm;
+			ComputeBlocksParams.BlockSize	 = Options.BlockSize;
+
+			FComputeBlocksResult ComputedBlocks = ComputeBlocks(OverlappedFile, ComputeBlocksParams);
+
+			const FGenericBlockArray& GenericBlocks = ComputedBlocks.Blocks;
 
 			FPath OutputFilename = Options.Output;
 			if (OutputFilename.empty())
diff --git a/Engine/Source/Programs/Unsync/Private/UnsyncCore.cpp b/Engine/Source/Programs/Unsync/Private/UnsyncCore.cpp
--- a/Engine/Source/Programs/Unsync/Private/UnsyncCore.cpp
+++ b/Engine/Source/Programs/Unsync/Private/UnsyncCore.cpp
@@ -122,9 +122,11 @@ ToBlock128(FGenericBlockArray& GenericBlocks)
 }
 
 template<typename WeakHasher>
-FGenericBlockArray
-ComputeBlocksVariableT(FIOReader& Reader, uint32 BlockSize, EStrongHashAlgorithmID StrongHasher, FComputeMacroBlockParams* OutMacroBlocks)
+FComputeBlocksResult
+ComputeBlocksVariableT(FIOReader& Reader, const FComputeBlocksParams& Params)
 {
+	const uint32 BlockSize = Params.BlockSize;
+
 	const uint64 InputSize = Reader.GetSize();
 
 	const uint32 MinimumBlockSize = ComputeMinVariableBlockSize(BlockSize);
@@ -133,9 +135,9 @@ ComputeBlocksVariableT(FIOReader& Reader, uint32 BlockSize, EStrongHashAlgorithm
 	const uint64 BytesPerTask = std::min<uint64>(256_MB, std::max<uint64>(InputSize, 1ull));  // TODO: handle task boundary overlap
 	const uint64 NumTasks	  = DivUp(InputSize, BytesPerTask);
 
-	const uint64 TargetMacroBlockSize	 = OutMacroBlocks ? OutMacroBlocks->TargetBlockSize : 0;
+	const uint64 TargetMacroBlockSize	 = Params.bNeedMacroBlocks ? Params.MacroBlockTargetSize : 0;
 	const uint64 MinimumMacroBlockSize	 = std::max<uint64>(MinimumBlockSize, TargetMacroBlockSize / 8);
-	const uint64 MaximumMacroBlockSize	 = OutMacroBlocks ? OutMacroBlocks->MaxBlockSize : 0;
+	const uint64 MaximumMacroBlockSize	 = Params.bNeedMacroBlocks ? Params.MacroBlockMaxSize : 0;
 	const uint32 BlocksPerMacroBlock	 = CheckedNarrow(DivUp(TargetMacroBlockSize - MinimumMacroBlockSize, BlockSize));
 	const uint32 MacroBlockHashThreshold = BlocksPerMacroBlock ? (0xFFFFFFFF / BlocksPerMacroBlock) : 0;
 
@@ -177,7 +179,7 @@ ComputeBlocksVariableT(FIOReader& Reader, uint32 BlockSize, EStrongHashAlgorithm
 		auto ScanTask = [&Tasks,
 						 &IoSemaphore,
 						 &BufferPool,
-						 StrongHasher,
+						 &Params,
 						 MinimumBlockSize,
 						 MaximumBlockSize,
 						 ScanTaskBuffer,
@@ -208,7 +210,7 @@ ComputeBlocksVariableT(FIOReader& Reader, uint32 BlockSize, EStrongHashAlgorithm
 						   &LastBlockEnd,
 						   &Task,
 						   DataBegin,
-						   StrongHasher,
+						   &Params,
 						   TargetMacroBlockSize,
 						   &MacroBlockHasher,
 						   &CurrentMacroBlock](const uint8* WindowBegin, const uint8* WindowEnd, uint32 WindowHash)
@@ -226,7 +228,7 @@ ComputeBlocksVariableT(FIOReader& Reader, uint32 BlockSize, EStrongHashAlgorithm
 									  Block.Offset	   = Task.Offset + uint64(LastBlockEnd - DataBegin);
 									  Block.Size	   = CheckedNarrow(ThisBlockSize);
 									  Block.HashWeak   = WindowHash;
-									  Block.HashStrong = ComputeHash(LastBlockEnd, ThisBlockSize, StrongHasher);
+									  Block.HashStrong = ComputeHash(LastBlockEnd, ThisBlockSize, Params.Algorithm.StrongHashAlgorithmId);
 
 									  if (TargetMacroBlockSize)
 									  {
@@ -282,24 +284,25 @@ ComputeBlocksVariableT(FIOReader& Reader, uint32 BlockSize, EStrongHashAlgorithm
 
 	// Merge blocks for all the tasks
 
-	FGenericBlockArray Result;
+	FComputeBlocksResult Result;
+
 	for (uint64 I = 0; I < NumTasks; ++I)
 	{
 		const FTask& Task = Tasks[I];
 		for (uint64 J = 0; J < Task.Blocks.size(); ++J)
 		{
-			Result.push_back(Task.Blocks[J]);
+			Result.Blocks.push_back(Task.Blocks[J]);
 		}
 	}
 
-	if (OutMacroBlocks)
+	if (Params.bNeedMacroBlocks)
 	{
 		for (uint64 I = 0; I < NumTasks; ++I)
 		{
 			const FTask& Task = Tasks[I];
 			for (uint64 J = 0; J < Task.MacroBlocks.size(); ++J)
 			{
-				OutMacroBlocks->Output.push_back(Task.MacroBlocks[J]);
+				Result.MacroBlocks.push_back(Task.MacroBlocks[J]);
 			}
 		}
 	}
@@ -317,12 +320,12 @@ ComputeBlocksVariableT(FIOReader& Reader, uint32 BlockSize, EStrongHashAlgorithm
 
 	THashSet<FGenericHash> UniqueBlockSet;
 	FGenericBlockArray	   UniqueBlocks;
-	for (const FGenericBlock& It : Result)
+	for (const FGenericBlock& It : Result.Blocks)
 	{
 		auto InsertResult = UniqueBlockSet.insert(It.HashStrong);
 		if (InsertResult.second)
 		{
-			if (It.Offset + It.Size < InputSize || Result.size() == 1)
+			if (It.Offset + It.Size < InputSize || Result.Blocks.size() == 1)
 			{
 				UniqueBlockMinSize = std::min<uint64>(UniqueBlockMinSize, It.Size);
 			}
@@ -365,38 +368,66 @@ ComputeBlocksVariableT(FIOReader& Reader, uint32 BlockSize, EStrongHashAlgorithm
 		(uint64)UniqueBlocks.size(),
 		NumTotalBlocks);
 
-	UNSYNC_ASSERT(NumTotalBlocks == Result.size());
+	UNSYNC_ASSERT(NumTotalBlocks == Result.Blocks.size());
 
 	return Result;
 }
 
-FGenericBlockArray
-ComputeBlocksVariable(FIOReader&				Reader,
-					  uint32					BlockSize,
-					  EWeakHashAlgorithmID		WeakHasher,
-					  EStrongHashAlgorithmID	StrongHasher,
-					  FComputeMacroBlockParams* OutMacroBlocks)
+FComputeBlocksResult
+ComputeBlocksVariable(FIOReader& Reader, const FComputeBlocksParams& Params)
 {
-	switch (WeakHasher)
+	switch (Params.Algorithm.WeakHashAlgorithmId)
 	{
 		case EWeakHashAlgorithmID::Naive:
-			return ComputeBlocksVariableT<FRollingChecksum>(Reader, BlockSize, StrongHasher, OutMacroBlocks);
+			return ComputeBlocksVariableT<FRollingChecksum>(Reader, Params);
 		case EWeakHashAlgorithmID::BuzHash:
-			return ComputeBlocksVariableT<FBuzHash>(Reader, BlockSize, StrongHasher, OutMacroBlocks);
+			return ComputeBlocksVariableT<FBuzHash>(Reader, Params);
 		default:
 			UNSYNC_FATAL(L"Unsupported weak hash algorithm mode");
 			return {};
 	}
 }
 
-template<typename WeakHasher>
 FGenericBlockArray
-ComputeBlocksFixedT(FIOReader& Reader, uint32 BlockSize, EStrongHashAlgorithmID StrongHasher, FComputeMacroBlockParams* OutMacroBlocks)
+ComputeBlocks(FIOReader& Reader, uint32 BlockSize, FAlgorithmOptions Algorithm)
+{
+	FComputeBlocksParams Params;
+	Params.Algorithm			= Algorithm;
+	Params.BlockSize			= BlockSize;
+	FComputeBlocksResult Result = ComputeBlocks(Reader, Params);
+	return std::move(Result.Blocks);
+}
+
+FGenericBlockArray
+ComputeBlocks(const uint8* Data, uint64 Size, uint32 BlockSize, FAlgorithmOptions Algorithm)
+{
+	FComputeBlocksParams Params;
+	Params.Algorithm			= Algorithm;
+	Params.BlockSize			= BlockSize;
+	FComputeBlocksResult Result = ComputeBlocks(Data, Size, Params);
+	return std::move(Result.Blocks);
+}
+
+FGenericBlockArray
+ComputeBlocksVariable(FIOReader& Reader, uint32 BlockSize, EWeakHashAlgorithmID WeakHasher, EStrongHashAlgorithmID StrongHasher)
+{
+	FComputeBlocksParams Params;
+	Params.Algorithm.WeakHashAlgorithmId   = WeakHasher;
+	Params.Algorithm.StrongHashAlgorithmId = StrongHasher;
+	Params.BlockSize					   = BlockSize;
+	FComputeBlocksResult Result			   = ComputeBlocks(Reader, Params);
+	return std::move(Result.Blocks);
+}
+
+template<typename WeakHasher>
+FComputeBlocksResult
+ComputeBlocksFixedT(FIOReader& Reader, const FComputeBlocksParams& Params)
 {
 	UNSYNC_LOG_INDENT;
 
 	auto TimeBegin = TimePointNow();
 
+	const uint32 BlockSize = Params.BlockSize;
 	const uint64 NumBlocks = DivUp(Reader.GetSize(), BlockSize);
 
 	FGenericBlockArray Blocks(NumBlocks);
@@ -408,10 +439,10 @@ ComputeBlocksFixedT(FIOReader& Reader, uint32 BlockSize, EStrongHashAlgorithmID
 	}
 
 	uint64 ReadSize = std::max<uint64>(BlockSize, 8_MB);
-	if (OutMacroBlocks)
+	if (Params.bNeedMacroBlocks)
 	{
 		UNSYNC_FATAL(L"Macro block generation is not implemented for fixed block mode");
-		ReadSize = std::max<uint64>(ReadSize, OutMacroBlocks->TargetBlockSize);
+		ReadSize = std::max<uint64>(ReadSize, Params.MacroBlockTargetSize);
 	}
 	UNSYNC_ASSERT(ReadSize % BlockSize == 0);
 
@@ -430,7 +461,7 @@ ComputeBlocksFixedT(FIOReader& Reader, uint32 BlockSize, EStrongHashAlgorithmID
 
 			IoSemaphore.Acquire();
 
-			auto ReadCallback = [&NumReadsCompleted, &TaskGroup, &NumBlocksCompleted, &Blocks, &IoSemaphore, StrongHasher, BlockSize](
+			auto ReadCallback = [&NumReadsCompleted, &TaskGroup, &NumBlocksCompleted, &Blocks, &IoSemaphore, &Params, BlockSize](
 									FIOBuffer CmdBuffer,
 									uint64	  CmdOffset,
 									uint64	  CmdReadSize,
@@ -441,7 +472,7 @@ ComputeBlocksFixedT(FIOReader& Reader, uint32 BlockSize, EStrongHashAlgorithmID
 							   &NumBlocksCompleted,
 							   &Blocks,
 							   &IoSemaphore,
-							   StrongHasher,
+							   &Params,
 							   BlockSize,
 							   CmdBuffer  = MakeShared(std::move(CmdBuffer)),
 							   BufferSize = CmdReadSize,
@@ -463,7 +494,7 @@ ComputeBlocksFixedT(FIOReader& Reader, uint32 BlockSize, EStrongHashAlgorithmID
 						UNSYNC_ASSERT(Block.HashWeak == 0);
 						UNSYNC_ASSERT(Block.HashStrong == FGenericHash{});
 
-						Block.HashStrong = ComputeHash(Buffer + I * BlockSize, Block.Size, StrongHasher);
+						Block.HashStrong = ComputeHash(Buffer + I * BlockSize, Block.Size, Params.Algorithm.StrongHashAlgorithmId);
 
 						WeakHasher HashWeak;
 						HashWeak.Update(Buffer + I * BlockSize, Block.Size);
@@ -515,22 +546,22 @@ ComputeBlocksFixedT(FIOReader& Reader, uint32 BlockSize, EStrongHashAlgorithmID
 		UniqueHashes.insert(It.HashWeak);
 	}
 
-	return Blocks;
+	FComputeBlocksResult Result;
+
+	std::swap(Result.Blocks, Blocks);
+
+	return Result;
 }
 
-FGenericBlockArray
-ComputeBlocksFixed(FIOReader&				 Reader,
-				   uint32					 BlockSize,
-				   EWeakHashAlgorithmID		 WeakHasher,
-				   EStrongHashAlgorithmID	 StrongHasher,
-				   FComputeMacroBlockParams* OutMacroBlocks)
+FComputeBlocksResult
+ComputeBlocksFixed(FIOReader& Reader, const FComputeBlocksParams& Params)
 {
-	switch (WeakHasher)
+	switch (Params.Algorithm.WeakHashAlgorithmId)
 	{
 		case EWeakHashAlgorithmID::Naive:
-			return ComputeBlocksFixedT<FRollingChecksum>(Reader, BlockSize, StrongHasher, OutMacroBlocks);
+			return ComputeBlocksFixedT<FRollingChecksum>(Reader, Params);
 		case EWeakHashAlgorithmID::BuzHash:
-			return ComputeBlocksFixedT<FBuzHash>(Reader, BlockSize, StrongHasher, OutMacroBlocks);
+			return ComputeBlocksFixedT<FBuzHash>(Reader, Params);
 		default:
 			UNSYNC_FATAL(L"Unsupported weak hash algorithm mode");
 			return {};
@@ -567,26 +598,26 @@ GetVersionString()
 	return Result;
 }
 
-FGenericBlockArray
-ComputeBlocks(FIOReader& Reader, uint32 BlockSize, FAlgorithmOptions Algorithm, FComputeMacroBlockParams* OutMacroBlocks)
+FComputeBlocksResult
+ComputeBlocks(FIOReader& Reader, const FComputeBlocksParams& Params)
 {
-	switch (Algorithm.ChunkingAlgorithmId)
+	switch (Params.Algorithm.ChunkingAlgorithmId)
 	{
 		case EChunkingAlgorithmID::FixedBlocks:
-			return ComputeBlocksFixed(Reader, BlockSize, Algorithm.WeakHashAlgorithmId, Algorithm.StrongHashAlgorithmId, OutMacroBlocks);
+			return ComputeBlocksFixed(Reader, Params);
 		case EChunkingAlgorithmID::VariableBlocks:
-			return ComputeBlocksVariable(Reader, BlockSize, Algorithm.WeakHashAlgorithmId, Algorithm.StrongHashAlgorithmId, OutMacroBlocks);
+			return ComputeBlocksVariable(Reader, Params);
 		default:
 			UNSYNC_FATAL(L"Unsupported chunking mode");
 			return {};
 	}
 }
 
-FGenericBlockArray
-ComputeBlocks(const uint8* Data, uint64 Size, uint32 BlockSize, FAlgorithmOptions Algorithm, FComputeMacroBlockParams* OutMacroBlocks)
+FComputeBlocksResult
+ComputeBlocks(const uint8* Data, uint64 Size, const FComputeBlocksParams& Params)
 {
 	FMemReader DataReader(Data, Size);
-	return ComputeBlocks(DataReader, BlockSize, Algorithm, OutMacroBlocks);
+	return ComputeBlocks(DataReader, Params);
 }
 
 template<typename BlockType>
@@ -1845,21 +1876,21 @@ CreateDirectoryManifest(const FPath& Root, uint32 BlockSize, FAlgorithmOptions A
 
 				Semaphore.Acquire();
 				TaskGroup.run([&Semaphore, &ResultMutex, &Result, File = std::move(File), Key = std::move(PathKey), BlockSize, Algorithm]() {
-					FComputeMacroBlockParams MacroBlocks;
+
+
+					FComputeBlocksParams ComputeBlocksParams;
+					ComputeBlocksParams.Algorithm = Algorithm;
+					ComputeBlocksParams.BlockSize = BlockSize;
 
 					// TODO: macro block generation is only implemented for variable chunk mode
-					const bool bNeedMacroBlocks = Algorithm.ChunkingAlgorithmId == EChunkingAlgorithmID::VariableBlocks;
+					ComputeBlocksParams.bNeedMacroBlocks = Algorithm.ChunkingAlgorithmId == EChunkingAlgorithmID::VariableBlocks;
 
-					FGenericBlockArray Blocks = ComputeBlocks(*File, BlockSize, Algorithm, bNeedMacroBlocks ? &MacroBlocks : nullptr);
+					FComputeBlocksResult ComputedBlocks = ComputeBlocks(*File, ComputeBlocksParams);
 
 					std::lock_guard<std::mutex> LockGuard(ResultMutex);
 
-					std::swap(Result.Files[Key].Blocks, Blocks);
-
-					if (bNeedMacroBlocks)
-					{
-						std::swap(Result.Files[Key].MacroBlocks, MacroBlocks.Output);
-					}
+					std::swap(Result.Files[Key].Blocks, ComputedBlocks.Blocks);
+					std::swap(Result.Files[Key].MacroBlocks, ComputedBlocks.MacroBlocks);
 
 					Semaphore.Release();
 				});
@@ -1991,18 +2022,19 @@ UpdateDirectoryManifestBlocks(FDirectoryManifest& Result, const FPath& Root, uin
 
 			Semaphore.Acquire();
 			TaskGroup.run([&FileManifest, &Semaphore, File = std::move(File), BlockSize, Algorithm]() {
-				FComputeMacroBlockParams MacroBlocks;
+
+				FComputeBlocksParams ComputeBlocksParams;
+				ComputeBlocksParams.Algorithm = Algorithm;
+				ComputeBlocksParams.BlockSize = BlockSize;
 
 				// TODO: macro block generation is only implemented for variable chunk mode
-				const bool bNeedMacroBlocks = Algorithm.ChunkingAlgorithmId == EChunkingAlgorithmID::VariableBlocks;
+				ComputeBlocksParams.bNeedMacroBlocks = Algorithm.ChunkingAlgorithmId == EChunkingAlgorithmID::VariableBlocks;
 
-				FileManifest.BlockSize = BlockSize;
-				FileManifest.Blocks	   = ComputeBlocks(*File, BlockSize, Algorithm, bNeedMacroBlocks ? &MacroBlocks : nullptr);
+				FComputeBlocksResult ComputedBlocks = ComputeBlocks(*File, ComputeBlocksParams);
+				std::swap(FileManifest.Blocks, ComputedBlocks.Blocks);
+				std::swap(FileManifest.MacroBlocks, ComputedBlocks.MacroBlocks);
 
-				if (bNeedMacroBlocks)
-				{
-					std::swap(FileManifest.MacroBlocks, MacroBlocks.Output);
-				}
+				FileManifest.BlockSize = BlockSize;
 
 				Semaphore.Release();
 			});
diff --git a/Engine/Source/Programs/Unsync/Private/UnsyncCore.h b/Engine/Source/Programs/Unsync/Private/UnsyncCore.h
--- a/Engine/Source/Programs/Unsync/Private/UnsyncCore.h
+++ b/Engine/Source/Programs/Unsync/Private/UnsyncCore.h
@@ -26,7 +26,6 @@ struct FRemoteDesc;
 struct FIOReader;
 struct FIOWriter;
 struct FIOReaderWriter;
-struct FComputeMacroBlockParams;
 
 struct FIdentityHash32
 {
@@ -108,16 +107,6 @@ FNeedListSize ComputeNeedListSize(const FNeedList& NeedList);
 
 using FGenericBlockArray = std::vector<FGenericBlock>;
 
-struct FComputeMacroBlockParams
-{
-	// Inputs
-	uint64 TargetBlockSize = 3_MB;
-	uint64 MaxBlockSize	   = 5_MB;	// Maximum allowed by Jupiter
-
-	// Outputs
-	FGenericBlockArray Output;
-};
-
 struct FFileManifest
 {
 	uint64 Mtime	 = 0;
@@ -129,6 +118,8 @@ struct FFileManifest
 	FGenericBlockArray MacroBlocks;
 	FPath			   CurrentPath;
 	bool			   bReadOnly = false;
+
+	bool IsValid() const { return Mtime != 0 && Size != 0; }
 };
 
 struct FAlgorithmOptionsV5
@@ -166,7 +157,8 @@ struct FDirectoryManifest
 
 	// wide string at runtime, utf8 serialized
 	// TODO: keep paths in canonical form (utf-8, unix-style separators)
-	std::unordered_map<std::wstring, FFileManifest> Files;
+	using FFileMap = std::unordered_map<std::wstring, FFileManifest>;
+	FFileMap Files;
 
 	// runtime data
 	FAlgorithmOptions Options = {};
@@ -196,22 +188,33 @@ void				   LogManifestFiles(ELogLevel LogLevel, const FDirectoryManifestInfo& In
 
 const std::string& GetVersionString();
 
-FGenericBlockArray ComputeBlocks(FIOReader&				   Reader,
-								 uint32					   BlockSize,
-								 FAlgorithmOptions		   Algorithm,
-								 FComputeMacroBlockParams* OutMacroBlocks = nullptr);
-
-FGenericBlockArray ComputeBlocks(const uint8*			   Data,
-								 uint64					   Size,
-								 uint32					   BlockSize,
-								 FAlgorithmOptions		   Algorithm,
-								 FComputeMacroBlockParams* OutMacroBlocks = nullptr);
-
-FGenericBlockArray ComputeBlocksVariable(FIOReader&				   Reader,
-										 uint32					   BlockSize,
-										 EWeakHashAlgorithmID	   WeakHasher,
-										 EStrongHashAlgorithmID	   StrongHasher,
-										 FComputeMacroBlockParams* OutMacroBlocks = nullptr);
+struct FComputeBlocksParams
+{
+	uint32			  BlockSize = 64_KB;
+	FAlgorithmOptions Algorithm;
+
+	bool   bNeedMacroBlocks		= false;
+	uint64 MacroBlockTargetSize = 3_MB;
+	uint64 MacroBlockMaxSize	= 5_MB;	 // Maximum allowed by Jupiter
+};
+
+struct FComputeBlocksResult
+{
+	FGenericBlockArray Blocks;
+	FGenericBlockArray MacroBlocks;
+};
+
+FComputeBlocksResult ComputeBlocks(FIOReader& Reader, const FComputeBlocksParams& Params);
+FComputeBlocksResult ComputeBlocks(const uint8* Data, uint64 Size, const FComputeBlocksParams& Params);
+FComputeBlocksResult ComputeBlocksVariable(FIOReader& Reader, const FComputeBlocksParams& Params);
+
+FGenericBlockArray ComputeBlocks(FIOReader& Reader, uint32 BlockSize, FAlgorithmOptions Algorithm);
+FGenericBlockArray ComputeBlocks(const uint8* Data, uint64 Size, uint32 BlockSize, FAlgorithmOptions Algorithm);
+FGenericBlockArray ComputeBlocksVariable(FIOReader&				Reader,
+										 uint32					BlockSize,
+										 EWeakHashAlgorithmID	WeakHasher,
+										 EStrongHashAlgorithmID StrongHasher);
+
 
 FNeedList DiffBlocks(const uint8*			   BaseData,
 					 uint64					   BaseDataSize,
