unsync - Use FComputeBlocksParams instead of loose parameters for CreateDirectoryManifest()

diff --git a/Engine/Source/Programs/Unsync/Private/UnsyncCmdHash.cpp b/Engine/Source/Programs/Unsync/Private/UnsyncCmdHash.cpp
--- a/Engine/Source/Programs/Unsync/Private/UnsyncCmdHash.cpp
+++ b/Engine/Source/Programs/Unsync/Private/UnsyncCmdHash.cpp
@@ -35,19 +35,26 @@ CmdHash(const FCmdHashOptions& Options)
 			DirectoryManifestPath = Options.Output;
 		}
 
+		FComputeBlocksParams ComputeBlocksParams;
+
+		ComputeBlocksParams.Algorithm = Options.Algorithm;
+		ComputeBlocksParams.BlockSize = Options.BlockSize;
+		// TODO: macro block generation is only implemented for variable chunk mode
+		ComputeBlocksParams.bNeedMacroBlocks = ComputeBlocksParams.Algorithm.ChunkingAlgorithmId == EChunkingAlgorithmID::VariableBlocks;
+
 		FDirectoryManifest DirectoryManifest;
 		if (Options.bForce)
 		{
-			DirectoryManifest = CreateDirectoryManifest(Options.Input, Options.BlockSize, Options.Algorithm);
+			DirectoryManifest = CreateDirectoryManifest(Options.Input, ComputeBlocksParams);
 		}
 		else if (Options.bIncremental)
 		{
 			UNSYNC_VERBOSE(L"Performing incremental directory manifest generation");
-			DirectoryManifest = CreateDirectoryManifestIncremental(Options.Input, Options.BlockSize, Options.Algorithm);
+			DirectoryManifest = CreateDirectoryManifestIncremental(Options.Input, ComputeBlocksParams);
 		}
 		else
 		{
-			LoadOrCreateDirectoryManifest(DirectoryManifest, Options.Input, Options.BlockSize, Options.Algorithm);
+			LoadOrCreateDirectoryManifest(DirectoryManifest, Options.Input, ComputeBlocksParams);
 		}
 
 		if (!GDryRun)
diff --git a/Engine/Source/Programs/Unsync/Private/UnsyncCmdPush.cpp b/Engine/Source/Programs/Unsync/Private/UnsyncCmdPush.cpp
--- a/Engine/Source/Programs/Unsync/Private/UnsyncCmdPush.cpp
+++ b/Engine/Source/Programs/Unsync/Private/UnsyncCmdPush.cpp
@@ -40,9 +40,13 @@ CmdPush(const FCmdPushOptions& Options)
 	AlgorithmOptions.WeakHashAlgorithmId   = EWeakHashAlgorithmID::BuzHash;
 	AlgorithmOptions.StrongHashAlgorithmId = EStrongHashAlgorithmID::Blake3_128;
 
+	FComputeBlocksParams ComputeBlocksParams;
+	ComputeBlocksParams.Algorithm = AlgorithmOptions;
+	ComputeBlocksParams.BlockSize = BlockSize;
+
 	FDirectoryManifest Manifest;
 	FPath			   ManifestPath	  = Options.Input / ".unsync" / "manifest.bin";	 // TODO: allow manifest path override
-	bool			   bManifestValid = LoadOrCreateDirectoryManifest(Manifest, Options.Input, BlockSize, AlgorithmOptions);
+	bool			   bManifestValid = LoadOrCreateDirectoryManifest(Manifest, Options.Input, ComputeBlocksParams);
 
 	if (!bManifestValid)
 	{
diff --git a/Engine/Source/Programs/Unsync/Private/UnsyncCore.cpp b/Engine/Source/Programs/Unsync/Private/UnsyncCore.cpp
--- a/Engine/Source/Programs/Unsync/Private/UnsyncCore.cpp
+++ b/Engine/Source/Programs/Unsync/Private/UnsyncCore.cpp
@@ -1818,13 +1818,13 @@ IsSynchronized(const FNeedList& NeedList, const FGenericBlockArray& SourceBlocks
 }
 
 FDirectoryManifest
-CreateDirectoryManifest(const FPath& Root, uint32 BlockSize, FAlgorithmOptions Algorithm)
+CreateDirectoryManifest(const FPath& Root, const FComputeBlocksParams& Params)
 {
 	UNSYNC_LOG_INDENT;
 
 	FDirectoryManifest Result;
 
-	Result.Options = Algorithm;
+	Result.Algorithm = Params.Algorithm;
 
 	FTimePoint TimeBegin = TimePointNow();
 
@@ -1858,7 +1858,7 @@ CreateDirectoryManifest(const FPath& Root, uint32 BlockSize, FAlgorithmOptions A
 		FileManifest.Mtime		 = ToWindowsFileTime(Dir.last_write_time());
 		FileManifest.Size		 = Dir.file_size();
 		FileManifest.CurrentPath = Dir.path();
-		FileManifest.BlockSize	 = BlockSize;
+		FileManifest.BlockSize	 = Params.BlockSize;
 		FileManifest.bReadOnly	 = IsReadOnly(Dir.status().permissions());
 
 		{
@@ -1866,7 +1866,7 @@ CreateDirectoryManifest(const FPath& Root, uint32 BlockSize, FAlgorithmOptions A
 			Result.Files[PathKey] = std::move(FileManifest);
 		}
 
-		if (BlockSize)
+		if (Params.bNeedBlocks && Params.BlockSize)
 		{
 			FPath FilePath = Root / RelativePath;
 			auto  File	   = std::make_shared<FNativeFile>(FilePath, EFileMode::ReadOnlyUnbuffered);
@@ -1875,17 +1875,9 @@ CreateDirectoryManifest(const FPath& Root, uint32 BlockSize, FAlgorithmOptions A
 				UNSYNC_VERBOSE(L"Computing blocks for '%ls' (%.2f MB)", FilePath.wstring().c_str(), double(File->GetSize()) / (1 << 20));
 
 				Semaphore.Acquire();
-				TaskGroup.run([&Semaphore, &ResultMutex, &Result, File = std::move(File), Key = std::move(PathKey), BlockSize, Algorithm]() {
+				TaskGroup.run([&Semaphore, &ResultMutex, &Result, File = std::move(File), Key = std::move(PathKey), &Params]() {
 
-
-					FComputeBlocksParams ComputeBlocksParams;
-					ComputeBlocksParams.Algorithm = Algorithm;
-					ComputeBlocksParams.BlockSize = BlockSize;
-
-					// TODO: macro block generation is only implemented for variable chunk mode
-					ComputeBlocksParams.bNeedMacroBlocks = Algorithm.ChunkingAlgorithmId == EChunkingAlgorithmID::VariableBlocks;
-
-					FComputeBlocksResult ComputedBlocks = ComputeBlocks(*File, ComputeBlocksParams);
+					FComputeBlocksResult ComputedBlocks = ComputeBlocks(*File, Params);
 
 					std::lock_guard<std::mutex> LockGuard(ResultMutex);
 
@@ -1906,7 +1898,7 @@ CreateDirectoryManifest(const FPath& Root, uint32 BlockSize, FAlgorithmOptions A
 
 	TaskGroup.wait();
 
-	if (BlockSize)
+	if (Params.bNeedBlocks && Params.BlockSize)
 	{
 		// TODO: move manifest stat printing into a helper
 
@@ -1938,8 +1930,10 @@ CreateDirectoryManifest(const FPath& Root, uint32 BlockSize, FAlgorithmOptions A
 }
 
 FDirectoryManifest
-CreateDirectoryManifestIncremental(const FPath& Root, uint32 BlockSize, FAlgorithmOptions Algorithm)
+CreateDirectoryManifestIncremental(const FPath& Root, const FComputeBlocksParams& InParams)
 {
+	FComputeBlocksParams Params = InParams;
+
 	FPath ManifestRoot			= Root / ".unsync";
 	FPath DirectoryManifestPath = ManifestRoot / "manifest.bin";
 
@@ -1949,11 +1943,15 @@ CreateDirectoryManifestIncremental(const FPath& Root, uint32 BlockSize, FAlgorit
 	// Inherit algorithm options from the existing manifest
 	if (bExistingManifestLoaded)
 	{
-		Algorithm = OldManifest.Options;
+		Params.Algorithm = OldManifest.Algorithm;
 	}
 
 	// Scan the input directory and gather file metadata, without generating blocks
-	FDirectoryManifest NewManifest = CreateDirectoryManifest(Root, 0, Algorithm);
+	FComputeBlocksParams LightweightManifestParams = Params; 
+	LightweightManifestParams.bNeedBlocks = false;
+	LightweightManifestParams.BlockSize = 0;
+
+	FDirectoryManifest NewManifest = CreateDirectoryManifest(Root, LightweightManifestParams);
 
 	// Copy file blocks from old manifest, if possible
 	for (auto& NewManifestFileEntry : NewManifest.Files)
@@ -1977,17 +1975,18 @@ CreateDirectoryManifestIncremental(const FPath& Root, uint32 BlockSize, FAlgorit
 	}
 
 	// Generate blocks for changed or new files
-	UpdateDirectoryManifestBlocks(NewManifest, Root, BlockSize, Algorithm);
+	UpdateDirectoryManifestBlocks(NewManifest, Root, Params);
 
 	return NewManifest;
 }
 
 void
-UpdateDirectoryManifestBlocks(FDirectoryManifest& Result, const FPath& Root, uint32 BlockSize, FAlgorithmOptions Algorithm)
+UpdateDirectoryManifestBlocks(FDirectoryManifest& Result, const FPath& Root, const FComputeBlocksParams& Params)
 {
 	UNSYNC_LOG_INDENT;
 
-	UNSYNC_ASSERT(BlockSize != 0);
+	UNSYNC_ASSERT(Params.bNeedBlocks);
+	UNSYNC_ASSERT(Params.BlockSize != 0);
 
 	FTimePoint TimeBegin = TimePointNow();
 
@@ -2004,7 +2003,7 @@ UpdateDirectoryManifestBlocks(FDirectoryManifest& Result, const FPath& Root, uin
 	for (auto& It : Result.Files)
 	{
 		FFileManifest& FileManifest = It.second;
-		if (FileManifest.BlockSize == BlockSize)
+		if (FileManifest.BlockSize == Params.BlockSize)
 		{
 			NumSkippedBlocks += FileManifest.Blocks.size();
 			NumSkippedBytes += FileManifest.Size;
@@ -2021,20 +2020,13 @@ UpdateDirectoryManifestBlocks(FDirectoryManifest& Result, const FPath& Root, uin
 			UNSYNC_VERBOSE(L"Computing blocks for '%ls' (%.2f MB)", FilePath.wstring().c_str(), double(File->GetSize()) / (1 << 20));
 
 			Semaphore.Acquire();
-			TaskGroup.run([&FileManifest, &Semaphore, File = std::move(File), BlockSize, Algorithm]() {
-
-				FComputeBlocksParams ComputeBlocksParams;
-				ComputeBlocksParams.Algorithm = Algorithm;
-				ComputeBlocksParams.BlockSize = BlockSize;
-
-				// TODO: macro block generation is only implemented for variable chunk mode
-				ComputeBlocksParams.bNeedMacroBlocks = Algorithm.ChunkingAlgorithmId == EChunkingAlgorithmID::VariableBlocks;
+			TaskGroup.run([&FileManifest, &Semaphore, &Params, File = std::move(File)]() {
 
-				FComputeBlocksResult ComputedBlocks = ComputeBlocks(*File, ComputeBlocksParams);
+				FComputeBlocksResult ComputedBlocks = ComputeBlocks(*File, Params);
 				std::swap(FileManifest.Blocks, ComputedBlocks.Blocks);
 				std::swap(FileManifest.MacroBlocks, ComputedBlocks.MacroBlocks);
 
-				FileManifest.BlockSize = BlockSize;
+				FileManifest.BlockSize = Params.BlockSize;
 
 				Semaphore.Release();
 			});
@@ -2095,7 +2087,7 @@ AlgorithmOptionsCompatible(const FAlgorithmOptions& A, const FAlgorithmOptions&
 }
 
 bool
-LoadOrCreateDirectoryManifest(FDirectoryManifest& Result, const FPath& Root, uint32 BlockSize, FAlgorithmOptions Algorithm)
+LoadOrCreateDirectoryManifest(FDirectoryManifest& Result, const FPath& Root, const FComputeBlocksParams& Params)
 {
 	UNSYNC_LOG_INDENT;
 
@@ -2112,18 +2104,29 @@ LoadOrCreateDirectoryManifest(FDirectoryManifest& Result, const FPath& Root, uin
 	}
 
 	const bool bExistingManifestLoaded = bManifestFileExists && LoadDirectoryManifest(OldDirectoryManifest, Root, DirectoryManifestPath);
-	const bool bExistingManifestCompatible = AlgorithmOptionsCompatible(OldDirectoryManifest.Options, Algorithm);
+	const bool bExistingManifestCompatible = AlgorithmOptionsCompatible(OldDirectoryManifest.Algorithm, Params.Algorithm);
 
 	if (bExistingManifestLoaded && bExistingManifestCompatible)
 	{
 		UNSYNC_VERBOSE(L"Loaded existing manifest from '%ls'", DirectoryManifestPath.wstring().c_str());
 
-		// Verify that manifests match in dry run mode.
-		// Otherwise just do a quick manifest generation, without file blocks.
-		uint32 NewManifestBlockSize = GDryRun ? BlockSize : 0;
-
 		UNSYNC_VERBOSE(L"Creating lightweight manifest for '%ls'", Root.wstring().c_str());
-		NewDirectoryManifest = CreateDirectoryManifest(Root, NewManifestBlockSize, Algorithm);
+		{
+			// Verify that manifests match in dry run mode.
+			// Otherwise just do a quick manifest generation, without file blocks.
+			FComputeBlocksParams NewParams = Params;
+			if (GDryRun)
+			{
+				NewParams.bNeedBlocks = true;
+			}
+			else
+			{
+				NewParams.bNeedBlocks = false;
+				NewParams.BlockSize	  = 0;
+			}
+
+			NewDirectoryManifest = CreateDirectoryManifest(Root, NewParams);
+		}
 
 		UNSYNC_VERBOSE(L"Comparing manifests");
 		for (const auto& OldManifestIt : OldDirectoryManifest.Files)
@@ -2159,16 +2162,16 @@ LoadOrCreateDirectoryManifest(FDirectoryManifest& Result, const FPath& Root, uin
 			}
 		}
 
-		if (BlockSize)
+		if (Params.bNeedBlocks && Params.BlockSize)
 		{
 			UNSYNC_VERBOSE(L"Updating manifest blocks");
-			UpdateDirectoryManifestBlocks(NewDirectoryManifest, Root, BlockSize, Algorithm);
+			UpdateDirectoryManifestBlocks(NewDirectoryManifest, Root, Params);
 		}
 	}
 	else
 	{
 		UNSYNC_VERBOSE(L"Creating manifest for '%ls'", Root.wstring().c_str());
-		NewDirectoryManifest = CreateDirectoryManifest(Root, BlockSize, Algorithm);
+		NewDirectoryManifest = CreateDirectoryManifest(Root, Params);
 	}
 
 	std::swap(Result, NewDirectoryManifest);
@@ -2201,9 +2204,9 @@ ComputeManifestStableSignature(const FDirectoryManifest& Manifest)
 	blake3_hasher Hasher;
 	blake3_hasher_init(&Hasher);
 
-	UpdateHashT(Hasher, Manifest.Options.ChunkingAlgorithmId);
-	UpdateHashT(Hasher, Manifest.Options.WeakHashAlgorithmId);
-	UpdateHashT(Hasher, Manifest.Options.StrongHashAlgorithmId);
+	UpdateHashT(Hasher, Manifest.Algorithm.ChunkingAlgorithmId);
+	UpdateHashT(Hasher, Manifest.Algorithm.WeakHashAlgorithmId);
+	UpdateHashT(Hasher, Manifest.Algorithm.StrongHashAlgorithmId);
 
 	std::vector<std::wstring> SortedFiles;
 	SortedFiles.reserve(Manifest.Files.size());
@@ -2861,7 +2864,7 @@ MergeManifests(FDirectoryManifest& Existing, const FDirectoryManifest& Other, bo
 		return true;
 	}
 
-	if (!AlgorithmOptionsCompatible(Existing.Options, Other.Options))
+	if (!AlgorithmOptionsCompatible(Existing.Algorithm, Other.Algorithm))
 	{
 		UNSYNC_ERROR("Trying to merge incompatible manifests (diff algorithm options do not match)");
 		return false;
@@ -3145,7 +3148,7 @@ LoadAndMergeSourceManifest(FDirectoryManifest& Output,
 		return false;
 	}
 
-	if (Output.IsValid() && !AlgorithmOptionsCompatible(Output.Options, LoadedManifest.Options))
+	if (Output.IsValid() && !AlgorithmOptionsCompatible(Output.Algorithm, LoadedManifest.Algorithm))
 	{
 		UNSYNC_ERROR(L"Can't merge manifest '%ls' as it uses different algorithm options", SourcePath.wstring().c_str());
 		return false;
@@ -3426,12 +3429,20 @@ SyncDirectory(const FSyncDirectoryOptions& SyncOptions)
 
 	ManifestLoadTimingLogger.Finish();
 
-	// Propagate algorithm selection from source
-	FAlgorithmOptions  Algorithm			   = SourceDirectoryManifest.Options;
-
 	FTimingLogger TargetManifestTimingLogger("Target directory manifest generation time");
 	UNSYNC_VERBOSE(L"Creating manifest for directory '%ls'", TargetPath.wstring().c_str());
-	FDirectoryManifest TargetDirectoryManifest = CreateDirectoryManifest(TargetPath, 0, Algorithm);
+
+
+	// Propagate algorithm selection from source
+	const FAlgorithmOptions Algorithm = SourceDirectoryManifest.Algorithm;
+
+	FComputeBlocksParams LightweightManifestParams;
+	LightweightManifestParams.Algorithm	  = Algorithm;
+	LightweightManifestParams.bNeedBlocks = false;
+	LightweightManifestParams.BlockSize	  = 0;
+
+	FDirectoryManifest TargetDirectoryManifest = CreateDirectoryManifest(TargetPath, LightweightManifestParams);
+
 	TargetManifestTimingLogger.Finish();
 
 	if (!bCaseSensitiveTargetFileSystem)
@@ -3480,11 +3491,11 @@ SyncDirectory(const FSyncDirectoryOptions& SyncOptions)
 	bool			   bBaseDirectoryManifestValid = false;
 	bool			   bQuickDifferencePossible	   = false;
 
-	if (!SyncOptions.bFullDifference && SourceDirectoryManifest.Options.ChunkingAlgorithmId == EChunkingAlgorithmID::VariableBlocks &&
+	if (!SyncOptions.bFullDifference && SourceDirectoryManifest.Algorithm.ChunkingAlgorithmId == EChunkingAlgorithmID::VariableBlocks &&
 		PathExists(BaseManifestPath))
 	{
 		bBaseDirectoryManifestValid = LoadDirectoryManifest(BaseDirectoryManifest, BasePath, BaseManifestPath);
-		if (bBaseDirectoryManifestValid && AlgorithmOptionsCompatible(SourceDirectoryManifest.Options, TargetDirectoryManifest.Options))
+		if (bBaseDirectoryManifestValid && AlgorithmOptionsCompatible(SourceDirectoryManifest.Algorithm, TargetDirectoryManifest.Algorithm))
 		{
 			bQuickDifferencePossible = true;
 		}
@@ -3786,7 +3797,7 @@ SyncDirectory(const FSyncDirectoryOptions& SyncOptions)
 
 			if (bProxyHasData)
 			{
-				ProxyPool.InitRequestMap(SourceDirectoryManifest.Options.StrongHashAlgorithmId);
+				ProxyPool.InitRequestMap(SourceDirectoryManifest.Algorithm.StrongHashAlgorithmId);
 
 				for (const FFileSyncTask& Task : AllFileTasks)
 				{
@@ -4691,7 +4702,8 @@ GetManifestInfo(const FDirectoryManifest& Manifest, bool bGenerateSignature)
 	Result.NumBlocks	  = UniqueBlockSet.size();
 	Result.NumMacroBlocks = UniqueMacroBlockSet.size();
 	Result.NumFiles		  = Manifest.Files.size();
-	Result.Algorithm	  = Manifest.Options;
+	Result.Algorithm	  = Manifest.Algorithm;
+
 	if (bGenerateSignature)
 	{
 		Result.StableSignature = ComputeManifestStableSignature(Manifest);
diff --git a/Engine/Source/Programs/Unsync/Private/UnsyncCore.h b/Engine/Source/Programs/Unsync/Private/UnsyncCore.h
--- a/Engine/Source/Programs/Unsync/Private/UnsyncCore.h
+++ b/Engine/Source/Programs/Unsync/Private/UnsyncCore.h
@@ -161,7 +161,7 @@ struct FDirectoryManifest
 	FFileMap Files;
 
 	// runtime data
-	FAlgorithmOptions Options = {};
+	FAlgorithmOptions Algorithm = {};
 	uint64			  Version = 0;
 
 	bool IsValid() const
@@ -190,7 +190,8 @@ const std::string& GetVersionString();
 
 struct FComputeBlocksParams
 {
-	uint32			  BlockSize = 64_KB;
+	bool			  bNeedBlocks = true;
+	uint32			  BlockSize	  = 64_KB;
 	FAlgorithmOptions Algorithm;
 
 	bool   bNeedMacroBlocks		= false;
@@ -413,10 +414,10 @@ struct FSyncDirectoryOptions
 
 bool SyncDirectory(const FSyncDirectoryOptions& SyncOptions);
 
-void UpdateDirectoryManifestBlocks(FDirectoryManifest& Result, const FPath& Root, uint32 BlockSize, FAlgorithmOptions Algorithm);
-FDirectoryManifest CreateDirectoryManifest(const FPath& Root, uint32 BlockSize, FAlgorithmOptions Algorithm);
-FDirectoryManifest CreateDirectoryManifestIncremental(const FPath& Root, uint32 BlockSize, FAlgorithmOptions Algorithm);
-bool LoadOrCreateDirectoryManifest(FDirectoryManifest& Result, const FPath& Root, uint32 BlockSize, FAlgorithmOptions Algorithm);
+void			   UpdateDirectoryManifestBlocks(FDirectoryManifest& Result, const FPath& Root, const FComputeBlocksParams& Params);
+FDirectoryManifest CreateDirectoryManifest(const FPath& Root, const FComputeBlocksParams& Params);
+FDirectoryManifest CreateDirectoryManifestIncremental(const FPath& Root, const FComputeBlocksParams& Params);
+bool			   LoadOrCreateDirectoryManifest(FDirectoryManifest& Result, const FPath& Root, const FComputeBlocksParams& Params);
 
 // Computes a Blake3 hash of the manifest blocks and file metadata, ignoring any other metadata.
 // Files are processed in sorted order, with file names treated as utf-8.
diff --git a/Engine/Source/Programs/Unsync/Private/UnsyncSerialization.cpp b/Engine/Source/Programs/Unsync/Private/UnsyncSerialization.cpp
--- a/Engine/Source/Programs/Unsync/Private/UnsyncSerialization.cpp
+++ b/Engine/Source/Programs/Unsync/Private/UnsyncSerialization.cpp
@@ -403,8 +403,8 @@ LoadDirectoryManifest(FDirectoryManifest& OutManifest, const FPath& Root, FIORea
 
 	if (Version >= FDirectoryManifest::EVersions::V6_VariableHash)
 	{
-		AlgorithmCompatibility = LoadOptionsSectionV5(Stream, OutManifest.Options);
-		HashSize			   = GetHashSize(ToHashType(OutManifest.Options.StrongHashAlgorithmId));
+		AlgorithmCompatibility = LoadOptionsSectionV5(Stream, OutManifest.Algorithm);
+		HashSize			   = GetHashSize(ToHashType(OutManifest.Algorithm.StrongHashAlgorithmId));
 	}
 
 	std::unordered_map<FHash256, FGenericBlockArray> MacroBlocks;
@@ -547,7 +547,7 @@ LoadDirectoryManifest(FDirectoryManifest& OutManifest, const FPath& Root, FIORea
 	// Old manifest versions always stored options after file blocks
 	if (Version < FDirectoryManifest::EVersions::V6_VariableHash)
 	{
-		AlgorithmCompatibility = LoadOptionsSectionV5(Stream, OutManifest.Options);
+		AlgorithmCompatibility = LoadOptionsSectionV5(Stream, OutManifest.Algorithm);
 	}
 
 	if (AlgorithmCompatibility != EAlgorithmCompatibilityResult::Ok)
@@ -557,7 +557,7 @@ LoadDirectoryManifest(FDirectoryManifest& OutManifest, const FPath& Root, FIORea
 	}
 
 	// Set the hash type in all file manifests
-	const EHashType HashType = ToHashType(OutManifest.Options.StrongHashAlgorithmId);
+	const EHashType HashType = ToHashType(OutManifest.Algorithm.StrongHashAlgorithmId);
 	for (auto& It : OutManifest.Files)
 	{
 		FFileManifest& FileManifest = It.second;
@@ -628,7 +628,7 @@ SaveDirectoryManifest(const FDirectoryManifest& Manifest, FVectorStreamOut& Stre
 {
 	// TODO: use compact binary to store the manifest
 
-	const size_t HashSize = GetHashSize(ToHashType(Manifest.Options.StrongHashAlgorithmId));
+	const size_t HashSize = GetHashSize(ToHashType(Manifest.Algorithm.StrongHashAlgorithmId));
 
 	FBuffer			 BlockStreamBuffer;
 	FVectorStreamOut BlockStream(BlockStreamBuffer);
@@ -638,7 +638,7 @@ SaveDirectoryManifest(const FDirectoryManifest& Manifest, FVectorStreamOut& Stre
 	Serialize(Stream, FDirectoryManifest::MAGIC);
 	Serialize(Stream, FDirectoryManifest::VERSION);
 
-	Serialize(Stream, Manifest.Options);
+	Serialize(Stream, Manifest.Algorithm);
 
 	// Save optional sections
 
