unsync - Groundwork for pack command

* This can be used to generate a compressed pack of a directory
* Supports using p4 have output to limit what is included in the pack

diff --git a/Engine/Source/Programs/Unsync/CMakeLists.txt b/Engine/Source/Programs/Unsync/CMakeLists.txt
--- a/Engine/Source/Programs/Unsync/CMakeLists.txt
+++ b/Engine/Source/Programs/Unsync/CMakeLists.txt
@@ -50,6 +50,8 @@ set(src
 	Private/UnsyncCmdQuery.h
 	Private/UnsyncCmdSync.cpp
 	Private/UnsyncCmdSync.h
+	Private/UnsyncCmdPack.cpp
+	Private/UnsyncCmdPack.h
 	Private/UnsyncCommon.h
 	Private/UnsyncCompression.cpp
 	Private/UnsyncCompression.h
diff --git a/Engine/Source/Programs/Unsync/Private/UnsyncCmdPack.cpp b/Engine/Source/Programs/Unsync/Private/UnsyncCmdPack.cpp
new file mode 100644
--- /dev/null
+++ b/Engine/Source/Programs/Unsync/Private/UnsyncCmdPack.cpp
@@ -0,0 +1,251 @@
+// Copyright Epic Games, Inc. All Rights Reserved.
+
+#pragma once
+
+#include "UnsyncCmdPack.h"
+#include "UnsyncFile.h"
+#include "UnsyncThread.h"
+#include "UnsyncSerialization.h"
+#include "UnsyncCompression.h"
+#include "UnsyncHashTable.h"
+
+#include <atomic>
+
+namespace unsync {
+
+template<typename CallbackT>
+static void
+ForLines(std::string_view String, CallbackT Callback)
+{
+	while (!String.empty())
+	{
+		size_t LineEndPos = String.find('\n');
+		if (LineEndPos == std::string::npos)
+		{
+			LineEndPos = String.length();
+		}
+
+		std::string_view LineView = String.substr(0, LineEndPos);
+
+		if (LineView.ends_with('\r'))
+		{
+			LineView = LineView.substr(0, LineView.length() - 1);
+		}
+
+		Callback(LineView);
+
+		String = String.substr(LineEndPos + 1);
+	}
+}
+
+static void
+BuildP4HaveSet(const FPath& Root, std::string_view P4HaveDataUtf8, FDirectoryManifest::FFileMap& Result)
+{
+	auto Callback = [&Result, &Root](std::string_view LineView)
+	{
+		if (LineView.starts_with("---")) // p4 diagnostic data
+		{
+			return;
+		}
+
+		size_t HashPos = LineView.find('#');
+		if (HashPos == std::string::npos)
+		{
+			return;
+		}
+
+		size_t SplitPos = LineView.find(" - ", HashPos);
+		if (SplitPos == std::string::npos)
+		{
+			return;
+		}
+
+		std::string_view DepotPathUtf8 = LineView.substr(0, SplitPos);
+		std::string_view LocalPathUtf8 = LineView.substr(SplitPos + 3);
+
+		FPath LocalPath(ConvertUtf8ToWide(LocalPathUtf8));
+
+		FPath RelativePath = GetRelativePath(LocalPath, Root);
+		if (RelativePath.empty())
+		{
+			return;
+		}
+
+		FFileManifest FileManifest;
+		FileManifest.CurrentPath = std::move(LocalPath);
+
+		std::wstring RelativePathStr = RelativePath.wstring();
+		Result.insert(std::make_pair(std::move(RelativePathStr), FileManifest));
+
+		UNSYNC_UNUSED(DepotPathUtf8);  // perhaps we should store this too?
+	};
+
+	ForLines(P4HaveDataUtf8, Callback);
+}
+
+struct FPackIndexEntry
+{
+	FHash160 Hash	 = {};
+	uint32	 Padding = 0;
+	uint64	 Offset  = 0;
+};
+static_assert(sizeof(FPackIndexEntry) == 32);
+
+int32 CmdPack(const FCmdPackOptions& Options)
+{
+	const FFileAttributes RootAttrib = GetFileAttrib(Options.RootPath);
+
+	const FPath InputRoot	 = Options.RootPath;
+	const FPath ManifestRoot = InputRoot / ".unsync";
+
+	if (!RootAttrib.bValid)
+	{
+		UNSYNC_ERROR(L"Input directory '%ls' does not exist", InputRoot.wstring().c_str());
+		return -1;
+	}
+
+	if (!RootAttrib.bDirectory)
+	{
+		UNSYNC_ERROR(L"Input '%ls' is not a directory", InputRoot.wstring().c_str());
+		return -1;
+	}
+
+	// TODO: allow explicit output path
+	const bool bOutputDirectoryCreated = (PathExists(ManifestRoot) && IsDirectory(ManifestRoot)) || CreateDirectories(ManifestRoot);
+	if (!bOutputDirectoryCreated)
+	{
+		UNSYNC_ERROR(L"Failed to create output directory '%ls'", ManifestRoot.wstring().c_str());
+		return -1;
+	}
+
+
+	const FPath DirectoryManifestPath = ManifestRoot / "manifest.bin";
+
+	std::string P4HaveBuffer;
+
+	FDirectoryManifest DirectoryManifest;
+
+	std::atomic<uint64> RawDataSize = 0;
+
+	if (!Options.P4HavePath.empty())
+	{
+		UNSYNC_VERBOSE(L"Loading p4 manifest file '%ls'", Options.P4HavePath.wstring().c_str());
+
+		FNativeFile P4HaveFile(Options.P4HavePath, EFileMode::ReadOnly);
+		if (!P4HaveFile.IsValid())
+		{
+			UNSYNC_ERROR(L"Could not open p4 manifest file '%ls'", Options.P4HavePath.wstring().c_str());
+			return -1;
+		}
+
+		P4HaveBuffer.resize(P4HaveFile.GetSize());
+		uint64 ReadBytes = P4HaveFile.Read(P4HaveBuffer.data(), 0, P4HaveFile.GetSize());
+		if (ReadBytes != P4HaveFile.GetSize())
+		{
+			UNSYNC_ERROR(L"Could not read the entire p4 manifest from '%ls'", Options.P4HavePath.wstring().c_str());
+			return -1;
+		}
+
+		DirectoryManifest.Algorithm = Options.Algorithm;
+
+		BuildP4HaveSet(InputRoot, P4HaveBuffer, DirectoryManifest.Files);
+
+		UNSYNC_VERBOSE(L"Loaded entries from p4 manifest: %llu", llu(DirectoryManifest.Files.size()));
+
+		UNSYNC_VERBOSE(L"Reading file attributes ...");
+		auto UpdateFileMetadata = [&RawDataSize](std::pair<const std::wstring, FFileManifest>& It)
+		{
+			FFileAttributes Attrib = GetFileAttrib(It.second.CurrentPath);
+
+			It.second.Mtime		= Attrib.Mtime;
+			It.second.Size		= Attrib.Size;
+			It.second.bReadOnly = true;	 // treat all p4 files as read-only in the manifest
+
+			RawDataSize += Attrib.Size;
+		};
+		ParallelForEach(DirectoryManifest.Files, UpdateFileMetadata);
+	}
+	else
+	{
+		// create a lightweight manifest, without blocks
+		FComputeBlocksParams LightweightManifestParams;
+		LightweightManifestParams.Algorithm	  = Options.Algorithm;
+		LightweightManifestParams.bNeedBlocks = false;
+		LightweightManifestParams.BlockSize	  = 0;
+
+		DirectoryManifest = CreateDirectoryManifest(InputRoot, LightweightManifestParams);
+	}
+
+	UNSYNC_VERBOSE(L"Found files: %llu", llu(DirectoryManifest.Files.size()));
+	UNSYNC_VERBOSE(L"Total file size: %llu bytes (%.2f MB)", llu(RawDataSize), SizeMb(RawDataSize));
+
+
+	FPath OutputPackFilename = ManifestRoot / "blocks.bin";
+	FNativeFile PackFile(OutputPackFilename, EFileMode::CreateWriteOnly);
+
+	FPath		OutputIndexFilename = ManifestRoot / "blocks.idx";
+	FNativeFile IndexFile(OutputIndexFilename, EFileMode::CreateWriteOnly);
+
+	UNSYNC_VERBOSE(L"Building file blocks ...");
+
+	FComputeBlocksParams BlockParams;
+	BlockParams.Algorithm = Options.Algorithm;
+	BlockParams.BlockSize = Options.BlockSize;
+	BlockParams.bNeedMacroBlocks = true;
+
+	std::atomic<uint64> PackFileOffset;
+	std::atomic<uint64> IndexFileOffset;
+
+	THashSet<FGenericHash> SeenMacroBlockHashSet;
+
+	std::mutex Mutex;
+	BlockParams.OnMacroBlockGenerated =
+		[&PackFile, &PackFileOffset, &IndexFile, &IndexFileOffset, &SeenMacroBlockHashSet, &Mutex](const FGenericBlock& Block,
+																								   FBufferView			Data)
+	{
+		{
+			std::lock_guard<std::mutex> LockGuard(Mutex);
+			if (!SeenMacroBlockHashSet.insert(Block.HashStrong).second)
+			{
+				return;
+			}
+		}
+
+		const uint64 MaxCompressedSize	  = GetMaxCompressedSize(Block.Size);
+		FIOBuffer	 CompressedData		  = FIOBuffer::Alloc(MaxCompressedSize, L"PackBlock");
+		uint64		 ActualCompressedSize = CompressInto(Data, CompressedData.GetMutBufferView(), 9);
+
+		if (ActualCompressedSize)
+		{
+			CompressedData.SetDataRange(0, ActualCompressedSize);
+
+			uint64 PackWriteOffset = PackFileOffset.fetch_add(CompressedData.GetSize());
+			PackFile.Write(CompressedData.GetData(), PackWriteOffset, CompressedData.GetSize());
+
+			FPackIndexEntry IndexEntry;
+			IndexEntry.Hash = Block.HashStrong.ToHash160();
+			IndexEntry.Offset = PackWriteOffset;
+
+			uint64 IndexWriteOffset = IndexFileOffset.fetch_add(sizeof(IndexEntry));
+			IndexFile.Write(&IndexEntry, IndexWriteOffset, sizeof(IndexEntry));
+		}
+	};
+
+	UpdateDirectoryManifestBlocks(DirectoryManifest, InputRoot, BlockParams);
+
+	if (!GDryRun)
+	{
+		UNSYNC_VERBOSE(L"Saving directory manifest '%ls'", DirectoryManifestPath.wstring().c_str());
+		SaveDirectoryManifest(DirectoryManifest, DirectoryManifestPath);
+	}
+
+	UNSYNC_VERBOSE(L"Total compressed size: %llu bytes (%.2f MB)", llu(PackFileOffset), SizeMb(PackFileOffset));
+
+	PackFile.Close();
+
+	UNSYNC_VERBOSE(L"Done!");
+
+	return 0;
+}
+
+}
diff --git a/Engine/Source/Programs/Unsync/Private/UnsyncCmdPack.h b/Engine/Source/Programs/Unsync/Private/UnsyncCmdPack.h
new file mode 100644
--- /dev/null
+++ b/Engine/Source/Programs/Unsync/Private/UnsyncCmdPack.h
@@ -0,0 +1,20 @@
+// Copyright Epic Games, Inc. All Rights Reserved.
+
+#pragma once
+
+#include "UnsyncCommon.h"
+#include "UnsyncCore.h"
+
+namespace unsync {
+
+struct FCmdPackOptions
+{
+	FPath RootPath;
+	FPath P4HavePath; // optional
+	uint32			  BlockSize = uint32(64_KB);
+	FAlgorithmOptions Algorithm;
+};
+
+int32 CmdPack(const FCmdPackOptions& Options);
+
+}
diff --git a/Engine/Source/Programs/Unsync/Private/UnsyncCmdPush.cpp b/Engine/Source/Programs/Unsync/Private/UnsyncCmdPush.cpp
--- a/Engine/Source/Programs/Unsync/Private/UnsyncCmdPush.cpp
+++ b/Engine/Source/Programs/Unsync/Private/UnsyncCmdPush.cpp
@@ -41,8 +41,9 @@ CmdPush(const FCmdPushOptions& Options)
 	AlgorithmOptions.StrongHashAlgorithmId = EStrongHashAlgorithmID::Blake3_128;
 
 	FComputeBlocksParams ComputeBlocksParams;
-	ComputeBlocksParams.Algorithm = AlgorithmOptions;
-	ComputeBlocksParams.BlockSize = BlockSize;
+	ComputeBlocksParams.Algorithm		 = AlgorithmOptions;
+	ComputeBlocksParams.BlockSize		 = BlockSize;
+	ComputeBlocksParams.bNeedMacroBlocks = true;
 
 	FDirectoryManifest Manifest;
 	FPath			   ManifestPath	  = Options.Input / ".unsync" / "manifest.bin";	 // TODO: allow manifest path override
diff --git a/Engine/Source/Programs/Unsync/Private/UnsyncCompression.cpp b/Engine/Source/Programs/Unsync/Private/UnsyncCompression.cpp
--- a/Engine/Source/Programs/Unsync/Private/UnsyncCompression.cpp
+++ b/Engine/Source/Programs/Unsync/Private/UnsyncCompression.cpp
@@ -5,13 +5,40 @@
 
 namespace unsync {
 
+uint64
+GetMaxCompressedSize(uint64 RawSize)
+{
+	return ZSTD_compressBound(RawSize);
+}
+
+uint64
+CompressInto(FBufferView Input, FMutBufferView Output, int ZstdCompressionLevel)
+{
+	const uint64 ExpectedMaxCompressedSize = GetMaxCompressedSize(Input.Size);
+	if (Output.Size < ExpectedMaxCompressedSize)
+	{
+		UNSYNC_FATAL(L"Compressed output buffer is too small");
+		return 0;
+	}
+
+	uint64 ActualCompressedSize = ZSTD_compress(Output.Data, Output.Size, Input.Data, Input.Size, ZstdCompressionLevel);
+	if (ZSTD_isError(ActualCompressedSize))
+	{
+		const char* ZstdError = ZSTD_getErrorName(ActualCompressedSize);
+		UNSYNC_ERROR(L"ZSTD_compress failed: %hs", ZstdError);
+		return 0;
+	}
+
+	return ActualCompressedSize;
+}
+
 FBuffer
 Compress(const uint8* Data, uint64 DataSize, int ZstdCompressionLevel)
 {
 	FBuffer Result;
 	if (DataSize)
 	{
-		Result.Resize(ZSTD_compressBound(DataSize));
+		Result.Resize(GetMaxCompressedSize(DataSize));
 
 		uint64 CompressedSize = ZSTD_compress(&Result[0], Result.Size(), Data, DataSize, ZstdCompressionLevel);
 		if (ZSTD_isError(CompressedSize))
diff --git a/Engine/Source/Programs/Unsync/Private/UnsyncCompression.h b/Engine/Source/Programs/Unsync/Private/UnsyncCompression.h
--- a/Engine/Source/Programs/Unsync/Private/UnsyncCompression.h
+++ b/Engine/Source/Programs/Unsync/Private/UnsyncCompression.h
@@ -7,7 +7,14 @@
 
 namespace unsync {
 
+uint64 GetMaxCompressedSize(uint64 RawSize);
+
 FBuffer Compress(const uint8* Data, uint64 DataSize, int ZstdCompressionLevel = 3);
+
+// Compress input into pre-allocated buffer that's at least as large as indicated by GetMaxCompressedSize()
+// Returns compressed data actual size.
+uint64 CompressInto(FBufferView Input, FMutBufferView Output, int ZstdCompressionLevel = 3);
+
 FBuffer Decompress(const uint8* Data, uint64 DataSize);
 FBuffer Decompress(const FBuffer& Buffer);
 bool	Decompress(const uint8* InputData, uint64 InputDataSize, uint8* OutputData, uint64 OutputDataSize);
diff --git a/Engine/Source/Programs/Unsync/Private/UnsyncCore.cpp b/Engine/Source/Programs/Unsync/Private/UnsyncCore.cpp
--- a/Engine/Source/Programs/Unsync/Private/UnsyncCore.cpp
+++ b/Engine/Source/Programs/Unsync/Private/UnsyncCore.cpp
@@ -27,7 +27,7 @@ UNSYNC_THIRD_PARTY_INCLUDES_START
 #include <md5-sse2.h>
 UNSYNC_THIRD_PARTY_INCLUDES_END
 
-#define UNSYNC_VERSION_STR "1.0.61-dev2"
+#define UNSYNC_VERSION_STR "1.0.61-dev3"
 
 namespace unsync {
 
@@ -247,6 +247,14 @@ ComputeBlocksVariableT(FIOReader& Reader, const FComputeBlocksParams& Params)
 																	 sizeof(CurrentMacroBlock.HashStrong.Data));
 											  Task.MacroBlocks.push_back(CurrentMacroBlock);
 
+											  if (Params.OnMacroBlockGenerated)
+											  {
+												  FBufferView BlockView;
+												  BlockView.Data = (LastBlockEnd + Block.Size) - CurrentMacroBlock.Size;
+												  BlockView.Size = CurrentMacroBlock.Size;
+												  Params.OnMacroBlockGenerated(CurrentMacroBlock, BlockView);
+											  }
+
 											  // Reset macro block state
 											  blake3_hasher_init(&MacroBlockHasher);
 											  CurrentMacroBlock.Offset += CurrentMacroBlock.Size;
@@ -254,6 +262,14 @@ ComputeBlocksVariableT(FIOReader& Reader, const FComputeBlocksParams& Params)
 										  }
 									  }
 
+									  if (Params.OnBlockGenerated)
+									  {
+										  FBufferView BlockView;
+										  BlockView.Data = LastBlockEnd;
+										  BlockView.Size = Block.Size;
+										  Params.OnBlockGenerated(Block, BlockView);
+									  }
+
 									  if (!Task.Blocks.empty())
 									  {
 										  UNSYNC_ASSERT(Task.Blocks.back().Offset + Task.Blocks.back().Size == Block.Offset);
@@ -2014,29 +2030,31 @@ UpdateDirectoryManifestBlocks(FDirectoryManifest& Result, const FPath& Root, con
 		++NumProcessedFiles;
 
 		FPath FilePath = Root / It.first;
-		auto  File	   = std::make_shared<FNativeFile>(FilePath, EFileMode::ReadOnlyUnbuffered);
-		if (File->IsValid())
-		{
-			UNSYNC_VERBOSE(L"Computing blocks for '%ls' (%.2f MB)", FilePath.wstring().c_str(), double(File->GetSize()) / (1 << 20));
 
-			Semaphore.Acquire();
-			TaskGroup.run([&FileManifest, &Semaphore, &Params, File = std::move(File)]() {
+		Semaphore.Acquire();
 
-				FComputeBlocksResult ComputedBlocks = ComputeBlocks(*File, Params);
+		UNSYNC_VERBOSE(L"Computing blocks for '%ls' (%.2f MB)", FilePath.wstring().c_str(), SizeMb(It.second.Size));
+		auto BlockTask = [&FileManifest, &Semaphore, &Params, FilePath = std::move(FilePath)]()
+		{
+			FNativeFile File(FilePath, EFileMode::ReadOnlyUnbuffered);
+			if (File.IsValid())
+			{
+				FComputeBlocksResult ComputedBlocks = ComputeBlocks(File, Params);
 				std::swap(FileManifest.Blocks, ComputedBlocks.Blocks);
 				std::swap(FileManifest.MacroBlocks, ComputedBlocks.MacroBlocks);
 
 				FileManifest.BlockSize = Params.BlockSize;
+			}
+			else
+			{
+				UNSYNC_FATAL(L"Failed to open file '%ls' while computing manifest blocks. %hs",
+							 FilePath.wstring().c_str(),
+							 FormatSystemErrorMessage(File.GetError()).c_str());
+			}
+			Semaphore.Release();
+		};
 
-				Semaphore.Release();
-			});
-		}
-		else
-		{
-			UNSYNC_FATAL(L"Failed to open file '%ls' while computing manifest blocks. %hs",
-						 FilePath.wstring().c_str(),
-						 FormatSystemErrorMessage(File->GetError()).c_str());
-		}
+		TaskGroup.run(BlockTask);
 	}
 
 	TaskGroup.wait();
diff --git a/Engine/Source/Programs/Unsync/Private/UnsyncCore.h b/Engine/Source/Programs/Unsync/Private/UnsyncCore.h
--- a/Engine/Source/Programs/Unsync/Private/UnsyncCore.h
+++ b/Engine/Source/Programs/Unsync/Private/UnsyncCore.h
@@ -13,6 +13,7 @@
 #include <unordered_map>
 #include <unordered_set>
 #include <vector>
+#include <functional>
 
 namespace unsync {
 
@@ -188,6 +189,8 @@ void				   LogManifestFiles(ELogLevel LogLevel, const FDirectoryManifestInfo& In
 
 const std::string& GetVersionString();
 
+using FOnBlockGenerated = std::function<void(const FGenericBlock& Block, FBufferView Data)>;
+
 struct FComputeBlocksParams
 {
 	bool			  bNeedBlocks = true;
@@ -197,6 +200,10 @@ struct FComputeBlocksParams
 	bool   bNeedMacroBlocks		= false;
 	uint64 MacroBlockTargetSize = 3_MB;
 	uint64 MacroBlockMaxSize	= 5_MB;	 // Maximum allowed by Jupiter
+
+	// Callbacks may be called from worker threads
+	FOnBlockGenerated OnBlockGenerated;
+	FOnBlockGenerated OnMacroBlockGenerated;
 };
 
 struct FComputeBlocksResult
diff --git a/Engine/Source/Programs/Unsync/Private/UnsyncFile.h b/Engine/Source/Programs/Unsync/Private/UnsyncFile.h
--- a/Engine/Source/Programs/Unsync/Private/UnsyncFile.h
+++ b/Engine/Source/Programs/Unsync/Private/UnsyncFile.h
@@ -105,6 +105,7 @@ struct FIOBuffer
 	}
 
 	FBufferView GetBufferView() const { return FBufferView{GetData(), GetSize()}; }
+	FMutBufferView GetMutBufferView() { return FMutBufferView{GetData(), GetSize()}; }
 
 private:
 	static constexpr uint64 CANARY = 0x67aced0423000de5ull;
diff --git a/Engine/Source/Programs/Unsync/Private/UnsyncMain.cpp b/Engine/Source/Programs/Unsync/Private/UnsyncMain.cpp
--- a/Engine/Source/Programs/Unsync/Private/UnsyncMain.cpp
+++ b/Engine/Source/Programs/Unsync/Private/UnsyncMain.cpp
@@ -5,6 +5,7 @@
 #include "UnsyncCmdHash.h"
 #include "UnsyncCmdLogin.h"
 #include "UnsyncCmdMount.h"
+#include "UnsyncCmdPack.h"
 #include "UnsyncCmdPatch.h"
 #include "UnsyncCmdPush.h"
 #include "UnsyncCmdQuery.h"
@@ -76,6 +77,7 @@ InnerMain(int Argc, char** Argv)
 	std::string				 QueryStringUtf8;
 	std::vector<std::string> QueryArgsUtf8;
 	std::string				 ScavengeRootUtf8;
+	std::string				 P4HavePathUtf8;
 	bool					 bForceOperation	 = false;
 	bool					 bAllowInsecureTls	 = false;
 	bool					 bUseTls			 = false;
@@ -150,6 +152,13 @@ InnerMain(int Argc, char** Argv)
 		"Create a directory manifest incrementally, by updating an existing manifest if one exists (only process changed files)");
 	SubCommands.push_back(SubHash);
 
+	// Configure pack
+
+	CLI::App* SubPack = Cli.add_subcommand("pack", "EXPERIMENTAL: Generate manifest for a directory and store all referenced data in a compressed pack file");
+	SubPack->add_option("Input", InputFilenameUtf8, "Input directory path")->required();
+	SubPack->add_option("--p4havefile", P4HavePathUtf8, "Use `p4 have` output from a given file to explicitly specify files included in the manifest");
+	SubCommands.push_back(SubPack);
+
 	// Configure push
 
 	CLI::App* SubPush = Cli.add_subcommand("push", "Loads a manifest from a directory and uploads referenced blocks to the remote server");
@@ -552,7 +561,7 @@ InnerMain(int Argc, char** Argv)
 	UNSYNC_VERBOSE2(L"Using threads: %d", GMaxThreads);
 	FConcurrencyPolicyScope ConcurrencyLimitScope(GMaxThreads);
 
-	if (Cli.got_subcommand(SubHash) || Cli.got_subcommand(SubSync))
+	if (Cli.got_subcommand(SubHash) || Cli.got_subcommand(SubSync) || Cli.got_subcommand(SubPack))
 	{
 		UNSYNC_VERBOSE(L"Using block size: %d KB", HashOrSyncBlockSize / 1024);
 	}
@@ -568,7 +577,7 @@ InnerMain(int Argc, char** Argv)
 		DefaultChunkingAlgorithm = EChunkingAlgorithmID::FixedBlocks;
 	}
 
-	if (Cli.got_subcommand(SubHash) || Cli.got_subcommand(SubDiff))
+	if (Cli.got_subcommand(SubHash) || Cli.got_subcommand(SubDiff) || Cli.got_subcommand(SubPack))
 	{
 		if (!bIncrementalMode)
 		{
@@ -744,6 +753,17 @@ InnerMain(int Argc, char** Argv)
 
 		return CmdHash(HashOptions);
 	}
+	if (Cli.got_subcommand(SubPack))
+	{
+		FCmdPackOptions PackOptions;
+
+		PackOptions.RootPath	 = InputFilename;
+		PackOptions.P4HavePath	 = NormalizeFilenameUtf8(P4HavePathUtf8);
+		PackOptions.BlockSize	 = HashOrSyncBlockSize;
+		PackOptions.Algorithm	 = Algorithm;
+
+		return CmdPack(PackOptions);
+	}
 	else if (Cli.got_subcommand(SubDiff))
 	{
 		FCmdDiffOptions DiffOptions;
