unsync - Store compressed block sizes in the pack index

diff --git a/Engine/Source/Programs/Unsync/Private/UnsyncCmdPack.cpp b/Engine/Source/Programs/Unsync/Private/UnsyncCmdPack.cpp
--- a/Engine/Source/Programs/Unsync/Private/UnsyncCmdPack.cpp
+++ b/Engine/Source/Programs/Unsync/Private/UnsyncCmdPack.cpp
@@ -86,7 +86,7 @@ BuildP4HaveSet(const FPath& Root, std::string_view P4HaveDataUtf8, FDirectoryMan
 struct FPackIndexEntry
 {
 	FHash160 Hash	 = {};
-	uint32	 Padding = 0;
+	uint32	 CompressedSize = 0;
 	uint64	 Offset  = 0;
 };
 static_assert(sizeof(FPackIndexEntry) == 32);
@@ -98,6 +98,9 @@ int32 CmdPack(const FCmdPackOptions& Options)
 	const FPath InputRoot	 = Options.RootPath;
 	const FPath ManifestRoot = InputRoot / ".unsync";
 
+	UNSYNC_VERBOSE(L"Generating package for directory '%ls' ...", InputRoot.wstring().c_str());
+	UNSYNC_LOG_INDENT;
+
 	if (!RootAttrib.bValid)
 	{
 		UNSYNC_ERROR(L"Input directory '%ls' does not exist", InputRoot.wstring().c_str());
@@ -125,8 +128,6 @@ int32 CmdPack(const FCmdPackOptions& Options)
 
 	FDirectoryManifest DirectoryManifest;
 
-	std::atomic<uint64> RawDataSize = 0;
-
 	if (!Options.P4HavePath.empty())
 	{
 		UNSYNC_VERBOSE(L"Loading p4 manifest file '%ls'", Options.P4HavePath.wstring().c_str());
@@ -153,15 +154,13 @@ int32 CmdPack(const FCmdPackOptions& Options)
 		UNSYNC_VERBOSE(L"Loaded entries from p4 manifest: %llu", llu(DirectoryManifest.Files.size()));
 
 		UNSYNC_VERBOSE(L"Reading file attributes ...");
-		auto UpdateFileMetadata = [&RawDataSize](std::pair<const std::wstring, FFileManifest>& It)
+		auto UpdateFileMetadata = [](std::pair<const std::wstring, FFileManifest>& It)
 		{
 			FFileAttributes Attrib = GetFileAttrib(It.second.CurrentPath);
 
 			It.second.Mtime		= Attrib.Mtime;
 			It.second.Size		= Attrib.Size;
 			It.second.bReadOnly = true;	 // treat all p4 files as read-only in the manifest
-
-			RawDataSize += Attrib.Size;
 		};
 		ParallelForEach(DirectoryManifest.Files, UpdateFileMetadata);
 	}
@@ -177,8 +176,6 @@ int32 CmdPack(const FCmdPackOptions& Options)
 	}
 
 	UNSYNC_VERBOSE(L"Found files: %llu", llu(DirectoryManifest.Files.size()));
-	UNSYNC_VERBOSE(L"Total file size: %llu bytes (%.2f MB)", llu(RawDataSize), SizeMb(RawDataSize));
-
 
 	FPath OutputPackFilename = ManifestRoot / "blocks.bin";
 	FNativeFile PackFile(OutputPackFilename, EFileMode::CreateWriteOnly);
@@ -219,12 +216,13 @@ int32 CmdPack(const FCmdPackOptions& Options)
 		{
 			CompressedData.SetDataRange(0, ActualCompressedSize);
 
-			uint64 PackWriteOffset = PackFileOffset.fetch_add(CompressedData.GetSize());
-			PackFile.Write(CompressedData.GetData(), PackWriteOffset, CompressedData.GetSize());
+			uint64 PackWriteOffset = PackFileOffset.fetch_add(ActualCompressedSize);
+			PackFile.Write(CompressedData.GetData(), PackWriteOffset, ActualCompressedSize);
 
 			FPackIndexEntry IndexEntry;
-			IndexEntry.Hash = Block.HashStrong.ToHash160();
-			IndexEntry.Offset = PackWriteOffset;
+			IndexEntry.Hash			  = Block.HashStrong.ToHash160();
+			IndexEntry.CompressedSize = CheckedNarrow<uint32>(ActualCompressedSize);
+			IndexEntry.Offset		  = PackWriteOffset;
 
 			uint64 IndexWriteOffset = IndexFileOffset.fetch_add(sizeof(IndexEntry));
 			IndexFile.Write(&IndexEntry, IndexWriteOffset, sizeof(IndexEntry));
@@ -233,17 +231,29 @@ int32 CmdPack(const FCmdPackOptions& Options)
 
 	UpdateDirectoryManifestBlocks(DirectoryManifest, InputRoot, BlockParams);
 
+	const uint64 CompressedSize = PackFileOffset;
+
 	if (!GDryRun)
 	{
 		UNSYNC_VERBOSE(L"Saving directory manifest '%ls'", DirectoryManifestPath.wstring().c_str());
 		SaveDirectoryManifest(DirectoryManifest, DirectoryManifestPath);
 	}
 
-	UNSYNC_VERBOSE(L"Total compressed size: %llu bytes (%.2f MB)", llu(PackFileOffset), SizeMb(PackFileOffset));
+	uint64 SourceSize = 0;
+	for (const auto& It : DirectoryManifest.Files)
+	{
+		SourceSize += It.second.Size;
+	}
 
-	PackFile.Close();
+	const uint64 NumSourceFiles = DirectoryManifest.Files.size();
+	UNSYNC_VERBOSE(L"Source files: %llu", llu(NumSourceFiles));
+	UNSYNC_VERBOSE(L"Source size: %llu bytes (%.2f MB)", llu(SourceSize), SizeMb(SourceSize));
+	UNSYNC_VERBOSE(L"Compressed size: %llu bytes (%.2f MB), %.0f%%",
+				   llu(CompressedSize),
+				   SizeMb(CompressedSize),
+				   100.0 * double(CompressedSize) / double(SourceSize));
 
-	UNSYNC_VERBOSE(L"Done!");
+	PackFile.Close();
 
 	return 0;
 }
