unsync - Tweak default logging level to output minimal information by default

* Add --silent flag to explicitly turn off all console logging except warnings and errors
* Output command line when starting logging to file

diff --git a/Engine/Source/Programs/Unsync/Private/UnsyncCmdDiff.cpp b/Engine/Source/Programs/Unsync/Private/UnsyncCmdDiff.cpp
--- a/Engine/Source/Programs/Unsync/Private/UnsyncCmdDiff.cpp
+++ b/Engine/Source/Programs/Unsync/Private/UnsyncCmdDiff.cpp
@@ -10,13 +10,13 @@ CmdDiff(const FCmdDiffOptions& Options)
 {
 	if (!Options.Output.empty())
 	{
-		UNSYNC_VERBOSE(L"Generating patch for '%ls' -> '%ls'", Options.Base.wstring().c_str(), Options.Source.wstring().c_str());
-		UNSYNC_VERBOSE(L"Output file '%ls'", Options.Output.wstring().c_str());
+		UNSYNC_LOG(L"Generating patch for '%ls' -> '%ls'", Options.Base.wstring().c_str(), Options.Source.wstring().c_str());
+		UNSYNC_LOG(L"Output file '%ls'", Options.Output.wstring().c_str());
 	}
 	else
 	{
-		UNSYNC_VERBOSE(L"Comparing base '%ls' and source '%ls'", Options.Base.wstring().c_str(), Options.Source.wstring().c_str());
-		UNSYNC_VERBOSE(L"Dry run mode (no output path given).");
+		UNSYNC_LOG(L"Comparing base '%ls' and source '%ls'", Options.Base.wstring().c_str(), Options.Source.wstring().c_str());
+		UNSYNC_LOG(L"Dry run mode (no output path given).");
 		GDryRun = true;
 	}
 
@@ -34,7 +34,7 @@ CmdDiff(const FCmdDiffOptions& Options)
 		return 1;
 	}
 
-	UNSYNC_VERBOSE(L"Generating patch");
+	UNSYNC_LOG(L"Generating patch");
 	FBuffer PatchData = GeneratePatch(BaseFile.Data(),
 									  BaseFile.Size(),
 									  SourceFile.Data(),
@@ -66,7 +66,7 @@ CmdDiff(const FCmdDiffOptions& Options)
 
 		if (!GDryRun && !Options.Output.empty())
 		{
-			UNSYNC_VERBOSE(L"Writing output file to '%ls'", Options.Output.wstring().c_str());
+			UNSYNC_LOG(L"Writing output file to '%ls'", Options.Output.wstring().c_str());
 			return WriteBufferToFile(Options.Output, PatchData.Data(), PatchData.Size()) ? 0 : 1;
 		}
 	}
diff --git a/Engine/Source/Programs/Unsync/Private/UnsyncCmdHash.cpp b/Engine/Source/Programs/Unsync/Private/UnsyncCmdHash.cpp
--- a/Engine/Source/Programs/Unsync/Private/UnsyncCmdHash.cpp
+++ b/Engine/Source/Programs/Unsync/Private/UnsyncCmdHash.cpp
@@ -12,7 +12,7 @@ CmdHash(const FCmdHashOptions& Options)
 {
 	if (unsync::IsDirectory(Options.Input))
 	{
-		UNSYNC_VERBOSE(L"Generating manifest for directory '%ls'", Options.Input.wstring().c_str());
+		UNSYNC_LOG(L"Generating manifest for directory '%ls'", Options.Input.wstring().c_str());
 
 		FPath InputRoot	   = Options.Input;
 		FPath ManifestRoot = InputRoot / ".unsync";
@@ -49,7 +49,7 @@ CmdHash(const FCmdHashOptions& Options)
 		}
 		else if (Options.bIncremental)
 		{
-			UNSYNC_VERBOSE(L"Performing incremental directory manifest generation");
+			UNSYNC_LOG(L"Performing incremental directory manifest generation");
 			DirectoryManifest = CreateDirectoryManifestIncremental(Options.Input, ComputeBlocksParams);
 		}
 		else
@@ -59,19 +59,19 @@ CmdHash(const FCmdHashOptions& Options)
 
 		if (!GDryRun)
 		{
-			UNSYNC_VERBOSE(L"Saving directory manifest '%ls'", DirectoryManifestPath.wstring().c_str());
+			UNSYNC_LOG(L"Saving directory manifest '%ls'", DirectoryManifestPath.wstring().c_str());
 
 			SaveDirectoryManifest(DirectoryManifest, DirectoryManifestPath);
 		}
 	}
 	else
 	{
-		UNSYNC_VERBOSE(L"Generating manfiest for file '%ls'", Options.Input.wstring().c_str());
+		UNSYNC_LOG(L"Generating manfiest for file '%ls'", Options.Input.wstring().c_str());
 
 		FNativeFile OverlappedFile(Options.Input);
 		if (OverlappedFile.IsValid())
 		{
-			UNSYNC_VERBOSE(L"Computing blocks for '%ls' (%.2f MB)", Options.Input.wstring().c_str(), SizeMb(OverlappedFile.GetSize()));
+			UNSYNC_LOG(L"Computing blocks for '%ls' (%.2f MB)", Options.Input.wstring().c_str(), SizeMb(OverlappedFile.GetSize()));
 			FComputeBlocksParams ComputeBlocksParams;
 			ComputeBlocksParams.Algorithm	 = Options.Algorithm;
 			ComputeBlocksParams.BlockSize	 = Options.BlockSize;
@@ -88,7 +88,7 @@ CmdHash(const FCmdHashOptions& Options)
 
 			if (!GDryRun)
 			{
-				UNSYNC_VERBOSE(L"Saving blocks to '%ls'", OutputFilename.wstring().c_str());
+				UNSYNC_LOG(L"Saving blocks to '%ls'", OutputFilename.wstring().c_str());
 				std::vector<FBlock128> Blocks128;
 				Blocks128.reserve(GenericBlocks.size());  // #wip-widehash
 				for (const auto& It : GenericBlocks)
diff --git a/Engine/Source/Programs/Unsync/Private/UnsyncCmdPack.cpp b/Engine/Source/Programs/Unsync/Private/UnsyncCmdPack.cpp
--- a/Engine/Source/Programs/Unsync/Private/UnsyncCmdPack.cpp
+++ b/Engine/Source/Programs/Unsync/Private/UnsyncCmdPack.cpp
@@ -98,7 +98,7 @@ int32 CmdPack(const FCmdPackOptions& Options)
 	const FPath InputRoot	 = Options.RootPath;
 	const FPath ManifestRoot = InputRoot / ".unsync";
 
-	UNSYNC_VERBOSE(L"Generating package for directory '%ls' ...", InputRoot.wstring().c_str());
+	UNSYNC_LOG(L"Generating package for directory '%ls' ...", InputRoot.wstring().c_str());
 	UNSYNC_LOG_INDENT;
 
 	if (!RootAttrib.bValid)
@@ -130,7 +130,7 @@ int32 CmdPack(const FCmdPackOptions& Options)
 
 	if (!Options.P4HavePath.empty())
 	{
-		UNSYNC_VERBOSE(L"Loading p4 manifest file '%ls'", Options.P4HavePath.wstring().c_str());
+		UNSYNC_LOG(L"Loading p4 manifest file '%ls'", Options.P4HavePath.wstring().c_str());
 
 		FNativeFile P4HaveFile(Options.P4HavePath, EFileMode::ReadOnly);
 		if (!P4HaveFile.IsValid())
@@ -151,9 +151,9 @@ int32 CmdPack(const FCmdPackOptions& Options)
 
 		BuildP4HaveSet(InputRoot, P4HaveBuffer, DirectoryManifest.Files);
 
-		UNSYNC_VERBOSE(L"Loaded entries from p4 manifest: %llu", llu(DirectoryManifest.Files.size()));
+		UNSYNC_LOG(L"Loaded entries from p4 manifest: %llu", llu(DirectoryManifest.Files.size()));
 
-		UNSYNC_VERBOSE(L"Reading file attributes ...");
+		UNSYNC_LOG(L"Reading file attributes ...");
 		auto UpdateFileMetadata = [](std::pair<const std::wstring, FFileManifest>& It)
 		{
 			FFileAttributes Attrib = GetFileAttrib(It.second.CurrentPath);
@@ -175,7 +175,7 @@ int32 CmdPack(const FCmdPackOptions& Options)
 		DirectoryManifest = CreateDirectoryManifest(InputRoot, LightweightManifestParams);
 	}
 
-	UNSYNC_VERBOSE(L"Found files: %llu", llu(DirectoryManifest.Files.size()));
+	UNSYNC_LOG(L"Found files: %llu", llu(DirectoryManifest.Files.size()));
 
 	FPath OutputPackFilename = ManifestRoot / "blocks.bin";
 	FNativeFile PackFile(OutputPackFilename, EFileMode::CreateWriteOnly);
@@ -183,7 +183,7 @@ int32 CmdPack(const FCmdPackOptions& Options)
 	FPath		OutputIndexFilename = ManifestRoot / "blocks.idx";
 	FNativeFile IndexFile(OutputIndexFilename, EFileMode::CreateWriteOnly);
 
-	UNSYNC_VERBOSE(L"Building file blocks ...");
+	UNSYNC_LOG(L"Building file blocks ...");
 
 	FComputeBlocksParams BlockParams;
 	BlockParams.Algorithm = Options.Algorithm;
@@ -235,7 +235,7 @@ int32 CmdPack(const FCmdPackOptions& Options)
 
 	if (!GDryRun)
 	{
-		UNSYNC_VERBOSE(L"Saving directory manifest '%ls'", DirectoryManifestPath.wstring().c_str());
+		UNSYNC_LOG(L"Saving directory manifest '%ls'", DirectoryManifestPath.wstring().c_str());
 		SaveDirectoryManifest(DirectoryManifest, DirectoryManifestPath);
 	}
 
@@ -246,9 +246,9 @@ int32 CmdPack(const FCmdPackOptions& Options)
 	}
 
 	const uint64 NumSourceFiles = DirectoryManifest.Files.size();
-	UNSYNC_VERBOSE(L"Source files: %llu", llu(NumSourceFiles));
-	UNSYNC_VERBOSE(L"Source size: %llu bytes (%.2f MB)", llu(SourceSize), SizeMb(SourceSize));
-	UNSYNC_VERBOSE(L"Compressed size: %llu bytes (%.2f MB), %.0f%%",
+	UNSYNC_LOG(L"Source files: %llu", llu(NumSourceFiles));
+	UNSYNC_LOG(L"Source size: %llu bytes (%.2f MB)", llu(SourceSize), SizeMb(SourceSize));
+	UNSYNC_LOG(L"Compressed size: %llu bytes (%.2f MB), %.0f%%",
 				   llu(CompressedSize),
 				   SizeMb(CompressedSize),
 				   100.0 * double(CompressedSize) / double(SourceSize));
diff --git a/Engine/Source/Programs/Unsync/Private/UnsyncCmdPatch.cpp b/Engine/Source/Programs/Unsync/Private/UnsyncCmdPatch.cpp
--- a/Engine/Source/Programs/Unsync/Private/UnsyncCmdPatch.cpp
+++ b/Engine/Source/Programs/Unsync/Private/UnsyncCmdPatch.cpp
@@ -8,8 +8,8 @@ namespace unsync {
 int32  // TODO: return a TResult
 CmdPatch(const FCmdPatchOptions& Options)
 {
-	UNSYNC_VERBOSE(L"Patching base file '%ls' using '%ls'", Options.Base.wstring().c_str(), Options.Patch.wstring().c_str());
-	UNSYNC_VERBOSE(L"Output file '%ls'", Options.Output.wstring().c_str());
+	UNSYNC_LOG(L"Patching base file '%ls' using '%ls'", Options.Base.wstring().c_str(), Options.Patch.wstring().c_str());
+	UNSYNC_LOG(L"Output file '%ls'", Options.Output.wstring().c_str());
 
 	FBuffer BaseFile = ReadFileToBuffer(Options.Base);
 	if (BaseFile.Empty())
@@ -34,7 +34,7 @@ CmdPatch(const FCmdPatchOptions& Options)
 
 	if (!GDryRun)
 	{
-		UNSYNC_VERBOSE(L"Writing output file to '%ls'", Options.Output.wstring().c_str());
+		UNSYNC_LOG(L"Writing output file to '%ls'", Options.Output.wstring().c_str());
 		return WriteBufferToFile(Options.Output, TargetData) ? 0 : 1;
 	}
 
diff --git a/Engine/Source/Programs/Unsync/Private/UnsyncCmdPush.cpp b/Engine/Source/Programs/Unsync/Private/UnsyncCmdPush.cpp
--- a/Engine/Source/Programs/Unsync/Private/UnsyncCmdPush.cpp
+++ b/Engine/Source/Programs/Unsync/Private/UnsyncCmdPush.cpp
@@ -119,7 +119,7 @@ CmdPush(const FCmdPushOptions& Options)
 
 		if (bPushComplete)
 		{
-			UNSYNC_VERBOSE(L"Push completed successfully", MaxAttempts);
+			UNSYNC_LOG(L"Push completed successfully", MaxAttempts);
 		}
 		else
 		{
diff --git a/Engine/Source/Programs/Unsync/Private/UnsyncCmdQuery.cpp b/Engine/Source/Programs/Unsync/Private/UnsyncCmdQuery.cpp
--- a/Engine/Source/Programs/Unsync/Private/UnsyncCmdQuery.cpp
+++ b/Engine/Source/Programs/Unsync/Private/UnsyncCmdQuery.cpp
@@ -212,15 +212,15 @@ CmdQuerySearch(const FCmdQueryOptions& Options)
 	const std::string& RootPath = Options.Args[0];
 	std::vector<std::regex> SubdirPatterns;
 
-	UNSYNC_VERBOSE("Searching '%hs'", RootPath.c_str());
+	UNSYNC_LOG("Searching '%hs'", RootPath.c_str());
 	if (Options.Args.size() > 1)
 	{
 		UNSYNC_LOG_INDENT;
-		UNSYNC_VERBOSE("Subdirectory patterns:", RootPath.c_str());
+		UNSYNC_LOG("Subdirectory patterns:", RootPath.c_str());
 		UNSYNC_LOG_INDENT;
 		for (size_t i = 1; i < Options.Args.size(); ++i)
 		{
-			UNSYNC_VERBOSE("%hs", Options.Args[i].c_str());
+			UNSYNC_LOG("%hs", Options.Args[i].c_str());
 			std::regex Pattern = std::regex(Options.Args[i], std::regex_constants::icase);
 			SubdirPatterns.push_back(Pattern);
 		}
diff --git a/Engine/Source/Programs/Unsync/Private/UnsyncCmdSync.cpp b/Engine/Source/Programs/Unsync/Private/UnsyncCmdSync.cpp
--- a/Engine/Source/Programs/Unsync/Private/UnsyncCmdSync.cpp
+++ b/Engine/Source/Programs/Unsync/Private/UnsyncCmdSync.cpp
@@ -28,29 +28,29 @@ CmdSync(const FCmdSyncOptions& Options)
 
 	if (Options.Source == ResolvedSource)
 	{
-		UNSYNC_VERBOSE(L"Sync source: '%ls'", Options.Source.wstring().c_str());
+		UNSYNC_LOG(L"Sync source: '%ls'", Options.Source.wstring().c_str());
 	}
 	else
 	{
-		UNSYNC_VERBOSE(L"Sync source: '%ls' ('%ls')", Options.Source.wstring().c_str(), ResolvedSource.wstring().c_str());
+		UNSYNC_LOG(L"Sync source: '%ls' ('%ls')", Options.Source.wstring().c_str(), ResolvedSource.wstring().c_str());
 	}
 
 	if (!Options.Filter->SyncIncludedWords.empty())
 	{
-		UNSYNC_VERBOSE(L"Include filter: ");
+		UNSYNC_LOG(L"Include filter: ");
 		UNSYNC_LOG_INDENT;
 		for( const std::wstring& include : Options.Filter->SyncIncludedWords)
 		{ 
-			UNSYNC_VERBOSE(L" %ls", include.c_str());
+			UNSYNC_LOG(L" %ls", include.c_str());
 		}
 	}
 	if(!Options.Filter->SyncExcludedWords.empty())
 	{
-		UNSYNC_VERBOSE(L"Exclude filter: ");
+		UNSYNC_LOG(L"Exclude filter: ");
 		UNSYNC_LOG_INDENT;
 		for( const std::wstring& exclude : Options.Filter->SyncExcludedWords)
 		{ 
-			UNSYNC_VERBOSE(L"%ls", exclude.c_str());
+			UNSYNC_LOG(L"%ls", exclude.c_str());
 		}
 	}
 
@@ -64,12 +64,12 @@ CmdSync(const FCmdSyncOptions& Options)
 		const FRemoteProtocolFeatures& Features = ProxyPool.GetFeatures();
 		if (Features.bFileDownload && Features.bDirectoryListing)
 		{
-			UNSYNC_VERBOSE(L"Server supports direct file access");
+			UNSYNC_LOG(L"Server supports direct file access");
 			bSourceFileSystemRequired = false;
 		}
 		else if (Features.bDownloadByHash && bSourceIsManifestHash)
 		{
-			UNSYNC_VERBOSE(L"Server supports access by manifest hash");
+			UNSYNC_LOG(L"Server supports access by manifest hash");
 			bSourceFileSystemRequired = false;
 		}
 		else
@@ -93,11 +93,11 @@ CmdSync(const FCmdSyncOptions& Options)
 			FPath ResolvedEntry = ResolvePath(Entry);
 			if (ResolvedEntry == Entry)
 			{
-				UNSYNC_VERBOSE(L"Sync overlay: '%ls'", ResolvedEntry.wstring().c_str());
+				UNSYNC_LOG(L"Sync overlay: '%ls'", ResolvedEntry.wstring().c_str());
 			}
 			else
 			{
-				UNSYNC_VERBOSE(L"Sync overlay: '%ls' ('%ls')", Entry.wstring().c_str(), ResolvedEntry.wstring().c_str());
+				UNSYNC_LOG(L"Sync overlay: '%ls' ('%ls')", Entry.wstring().c_str(), ResolvedEntry.wstring().c_str());
 			}
 			ResolvedOverlays.push_back(ResolvedEntry);
 		}
@@ -124,11 +124,11 @@ CmdSync(const FCmdSyncOptions& Options)
 		}
 	}
 
-	UNSYNC_VERBOSE(L"Sync target: '%ls'", Options.Target.wstring().c_str());
+	UNSYNC_LOG(L"Sync target: '%ls'", Options.Target.wstring().c_str());
 
 	if (!Options.SourceManifestOverride.empty())
 	{
-		UNSYNC_VERBOSE(L"Manifest override: %ls", Options.SourceManifestOverride.wstring().c_str());
+		UNSYNC_LOG(L"Manifest override: %ls", Options.SourceManifestOverride.wstring().c_str());
 
 		if (Options.Remote.IsValid() && !Options.bFullSourceScan)
 		{
@@ -144,11 +144,11 @@ CmdSync(const FCmdSyncOptions& Options)
 		{
 			if (bSourceIsDirectory)
 			{
-				UNSYNC_VERBOSE(L"'%ls' is a directory", Options.Source.wstring().c_str());
+				UNSYNC_LOG(L"'%ls' is a directory", Options.Source.wstring().c_str());
 			}
 			else
 			{
-				UNSYNC_VERBOSE(L"Assuming '%ls' is a directory", Options.Source.wstring().c_str());
+				UNSYNC_LOG(L"Assuming '%ls' is a directory", Options.Source.wstring().c_str());
 			}
 
 			FSyncDirectoryOptions SyncOptions;
@@ -185,11 +185,11 @@ CmdSync(const FCmdSyncOptions& Options)
 		}
 		else
 		{
-			UNSYNC_VERBOSE(L"'%ls' is a file", Options.Source.wstring().c_str());
+			UNSYNC_LOG(L"'%ls' is a file", Options.Source.wstring().c_str());
 
 			FSyncFileOptions SyncFileOptions;
 			SyncFileOptions.Algorithm			 = Options.Algorithm;
-			SyncFileOptions.BlockSize			 = Options.BlockSize;
+			SyncFileOptions.BlockSize			 = 64_KB;
 			SyncFileOptions.bValidateTargetFiles = Options.bValidateTargetFiles;
 
 			return SyncFile(Options.Source, Options.Target, Options.Target, SyncFileOptions).Succeeded() ? 0 : 1;
diff --git a/Engine/Source/Programs/Unsync/Private/UnsyncCmdSync.h b/Engine/Source/Programs/Unsync/Private/UnsyncCmdSync.h
--- a/Engine/Source/Programs/Unsync/Private/UnsyncCmdSync.h
+++ b/Engine/Source/Programs/Unsync/Private/UnsyncCmdSync.h
@@ -30,8 +30,6 @@ struct FCmdSyncOptions
 
 	bool bValidateTargetFiles = true;  // WARNING: turning this off is intended only for testing/profiling
 
-	uint32 BlockSize = uint32(64_KB);
-
 	FSyncFilter* Filter = nullptr;
 };
 
diff --git a/Engine/Source/Programs/Unsync/Private/UnsyncCore.cpp b/Engine/Source/Programs/Unsync/Private/UnsyncCore.cpp
--- a/Engine/Source/Programs/Unsync/Private/UnsyncCore.cpp
+++ b/Engine/Source/Programs/Unsync/Private/UnsyncCore.cpp
@@ -27,7 +27,7 @@ UNSYNC_THIRD_PARTY_INCLUDES_START
 #include <md5-sse2.h>
 UNSYNC_THIRD_PARTY_INCLUDES_END
 
-#define UNSYNC_VERSION_STR "1.0.61-dev3"
+#define UNSYNC_VERSION_STR "1.0.61-dev4"
 
 namespace unsync {
 
@@ -3331,7 +3331,7 @@ SyncDirectory(const FSyncDirectoryOptions& SyncOptions)
 
 	if (SyncOptions.bCleanup)
 	{
-		UNSYNC_VERBOSE(L"Unnecessary files will be deleted after sync (cleanup mode)");
+		UNSYNC_LOG(L"Unnecessary files will be deleted after sync (cleanup mode)");
 	}
 
 	FPath BaseManifestRoot = BasePath / ".unsync";
@@ -3379,7 +3379,7 @@ SyncDirectory(const FSyncDirectoryOptions& SyncOptions)
 
 	const bool bCaseSensitiveTargetFileSystem = IsCaseSensitiveFileSystem(TargetTempPath);
 
-	FTimingLogger ManifestLoadTimingLogger("Manifest load time");
+	FTimingLogger ManifestLoadTimingLogger("Manifest load time", ELogLevel::Info);
 
 	if (SyncOptions.SourceType == ESyncSourceType::ServerWithManifestHash)
 	{
@@ -3389,7 +3389,7 @@ SyncDirectory(const FSyncDirectoryOptions& SyncOptions)
 			return false;
 		}
 
-		UNSYNC_VERBOSE(L"Downloading manifest ...");
+		UNSYNC_LOG(L"Downloading manifest ...");
 
 		std::unique_ptr<FProxy> Proxy = ProxyPool.Alloc();
 
@@ -3447,8 +3447,8 @@ SyncDirectory(const FSyncDirectoryOptions& SyncOptions)
 
 	ManifestLoadTimingLogger.Finish();
 
-	FTimingLogger TargetManifestTimingLogger("Target directory manifest generation time");
-	UNSYNC_VERBOSE(L"Creating manifest for directory '%ls'", TargetPath.wstring().c_str());
+	FTimingLogger TargetManifestTimingLogger("Target directory manifest generation time", ELogLevel::Info);
+	UNSYNC_LOG(L"Creating manifest for directory '%ls'", TargetPath.wstring().c_str());
 
 
 	// Propagate algorithm selection from source
@@ -3490,15 +3490,15 @@ SyncDirectory(const FSyncDirectoryOptions& SyncOptions)
 	std::vector<FFileSyncTask> AllFileTasks;
 
 	LogGlobalStatus(L"Scanning base directory");
-	UNSYNC_VERBOSE(L"Scanning base directory");
+	UNSYNC_LOG(L"Scanning base directory");
 	FFileAttributeCache BaseAttribCache = CreateFileAttributeCache(BasePath, SyncFilter);
-	UNSYNC_VERBOSE(L"Base files: %d", (uint32)BaseAttribCache.Map.size());
+	UNSYNC_LOG(L"Base files: %d", (uint32)BaseAttribCache.Map.size());
 
 	FFileAttributeCache SourceAttribCache;
 	if (bFileSystemSource && SyncOptions.bValidateSourceFiles)
 	{
 		LogGlobalStatus(L"Scanning source directory");
-		UNSYNC_VERBOSE(L"Scanning source directory");
+		UNSYNC_LOG(L"Scanning source directory");
 		SourceAttribCache = CreateFileAttributeCache(SourcePath, SyncFilter);
 	}
 
@@ -3521,7 +3521,7 @@ SyncDirectory(const FSyncDirectoryOptions& SyncOptions)
 
 	if (bQuickDifferencePossible)
 	{
-		UNSYNC_VERBOSE(L"Quick file difference is allowed (use --full-diff option to override)");
+		UNSYNC_LOG(L"Quick file difference is allowed (use --full-diff option to override)");
 	}
 
 	uint64 TotalSourceSize = 0;
@@ -3668,7 +3668,7 @@ SyncDirectory(const FSyncDirectoryOptions& SyncOptions)
 	}
 
 	LogGlobalStatus(L"Computing difference");
-	UNSYNC_VERBOSE(L"Computing difference ...");
+	UNSYNC_LOG(L"Computing difference ...");
 
 	uint64 EstimatedNeedBytesFromSource = 0;
 	uint64 EstimatedNeedBytesFromBase	= 0;
@@ -3741,7 +3741,7 @@ SyncDirectory(const FSyncDirectoryOptions& SyncOptions)
 		auto TimeDiffEnd = TimePointNow();
 
 		double Duration = DurationSec(TimeDiffBegin, TimeDiffEnd);
-		UNSYNC_VERBOSE(L"Difference complete in %.3f sec", Duration);
+		UNSYNC_LOG(L"Difference complete in %.3f sec", Duration);
 
 		for (FFileSyncTask& Item : AllFileTasks)
 		{
@@ -3750,8 +3750,8 @@ SyncDirectory(const FSyncDirectoryOptions& SyncOptions)
 			TotalSyncSizeBytes += Item.TotalSizeBytes;
 		}
 
-		UNSYNC_VERBOSE(L"Total need from source: %.2f MB", SizeMb(EstimatedNeedBytesFromSource));
-		UNSYNC_VERBOSE(L"Total need from base: %.2f MB", SizeMb(EstimatedNeedBytesFromBase));
+		UNSYNC_LOG(L"Total need from source: %.2f MB", SizeMb(EstimatedNeedBytesFromSource));
+		UNSYNC_LOG(L"Total need from base: %.2f MB", SizeMb(EstimatedNeedBytesFromBase));
 
 		uint64 AvailableDiskBytes = SyncOptions.bCheckAvailableSpace ? GetAvailableDiskSpace(TargetPath) : ~0ull;
 		if (TotalSyncSizeBytes > AvailableDiskBytes)
@@ -3774,7 +3774,7 @@ SyncDirectory(const FSyncDirectoryOptions& SyncOptions)
 	std::unique_ptr<FScavengeDatabase> ScavengeDatabase;
 	if (!SyncOptions.ScavengeRoot.empty())
 	{
-		UNSYNC_VERBOSE(L"Scavenging blocks from existing data sets");
+		UNSYNC_LOG(L"Scavenging blocks from existing data sets");
 		UNSYNC_LOG_INDENT;
 
 		FTimePoint ScavengeDbTimeBegin = TimePointNow();
@@ -3784,7 +3784,7 @@ SyncDirectory(const FSyncDirectoryOptions& SyncOptions)
 
 		double Duration = DurationSec(ScavengeDbTimeBegin, TimePointNow());
 
-		UNSYNC_VERBOSE(L"Done in %.3f sec", Duration);
+		UNSYNC_LOG(L"Done in %.3f sec", Duration);
 	}
 
 	LogGlobalProgress();
@@ -3792,7 +3792,7 @@ SyncDirectory(const FSyncDirectoryOptions& SyncOptions)
 	if (ProxyPool.IsValid())
 	{
 		LogGlobalStatus(L"Connecting to server");
-		UNSYNC_VERBOSE(L"Connecting to %hs server '%hs:%d' ...",
+		UNSYNC_LOG(L"Connecting to %hs server '%hs:%d' ...",
 					   ToString(ProxyPool.RemoteDesc.Protocol),
 					   ProxyPool.RemoteDesc.HostAddress.c_str(),
 					   ProxyPool.RemoteDesc.HostPort);
@@ -3804,10 +3804,10 @@ SyncDirectory(const FSyncDirectoryOptions& SyncOptions)
 		{
 			// TODO: report TLS status
 			// ESocketSecurity security = proxy->get_socket_security();
-			// UNSYNC_VERBOSE(L"Connection established (security: %hs)", ToString(security));
+			// UNSYNC_LOG(L"Connection established (security: %hs)", ToString(security));
 
-			UNSYNC_VERBOSE(L"Connection established");
-			UNSYNC_VERBOSE(L"Building block request map");
+			UNSYNC_LOG(L"Connection established");
+			UNSYNC_LOG(L"Building block request map");
 
 			const bool bProxyHasData = Proxy->Contains(SourceDirectoryManifest);
 
@@ -3842,7 +3842,7 @@ SyncDirectory(const FSyncDirectoryOptions& SyncOptions)
 	}
 
 	LogGlobalStatus(L"Copying files");
-	UNSYNC_VERBOSE(L"Copying files");
+	UNSYNC_LOG(L"Copying files ...");
 
 	{
 		// Throttle background tasks by trying to keep them to some sensible memory budget. Best effort only, not a hard limit.
@@ -4104,7 +4104,7 @@ SyncDirectory(const FSyncDirectoryOptions& SyncOptions)
 
 		if (NumBackgroundTasks != 0)
 		{
-			UNSYNC_VERBOSE(L"Waiting for background tasks to complete");
+			UNSYNC_LOG(L"Waiting for background tasks to complete");
 		}
 		BackgroundTaskGroup.wait();
 
@@ -4155,7 +4155,7 @@ SyncDirectory(const FSyncDirectoryOptions& SyncOptions)
 
 	if (bSyncSucceeded && SyncOptions.bCleanup)
 	{
-		UNSYNC_VERBOSE(L"Deleting unnecessary files");
+		UNSYNC_LOG(L"Deleting unnecessary files");
 		UNSYNC_LOG_INDENT;
 		DeleteUnnecessaryFiles(TargetPath, TargetDirectoryManifest, SourceDirectoryManifest, SyncFilter);
 	}
@@ -4172,9 +4172,9 @@ SyncDirectory(const FSyncDirectoryOptions& SyncOptions)
 		}
 	}
 
-	UNSYNC_VERBOSE(L"Skipped files: %d, full copies: %d, partial copies: %d", StatSkipped, StatFullCopy, StatPartialCopy);
-	UNSYNC_VERBOSE(L"Copied from source: %.2f MB, copied from base: %.2f MB", SizeMb(StatSourceBytes), SizeMb(StatBaseBytes));
-	UNSYNC_VERBOSE(L"Sync completed %ls", bSyncSucceeded ? L"successfully" : L"with errors (see log for details)");
+	UNSYNC_LOG(L"Skipped files: %d, full copies: %d, partial copies: %d", StatSkipped, StatFullCopy, StatPartialCopy);
+	UNSYNC_LOG(L"Copied from source: %.2f MB, copied from base: %.2f MB", SizeMb(StatSourceBytes), SizeMb(StatBaseBytes));
+	UNSYNC_LOG(L"Sync completed %ls", bSyncSucceeded ? L"successfully" : L"with errors (see log for details)");
 
 	double ElapsedSeconds = DurationSec(TimeBegin, TimePointNow());
 	UNSYNC_VERBOSE2(L"Sync time: %.2f seconds", ElapsedSeconds);
diff --git a/Engine/Source/Programs/Unsync/Private/UnsyncLog.cpp b/Engine/Source/Programs/Unsync/Private/UnsyncLog.cpp
--- a/Engine/Source/Programs/Unsync/Private/UnsyncLog.cpp
+++ b/Engine/Source/Programs/Unsync/Private/UnsyncLog.cpp
@@ -41,6 +41,7 @@ IsDebuggerPresent()
 
 thread_local bool	GLogVerbose			= false;
 bool				GLogVeryVerbose		= false;
+bool				GLogSilent			= false;
 bool				GBreakOnError		= IsDebuggerPresent();
 bool				GBreakOnWarning		= false;
 thread_local uint32 GLogIndent			= 0;
@@ -54,6 +55,18 @@ thread_local uint32 GLogThreadIndex = ~0u;
 
 FTimePoint GNextFlushTime = TimePointNow();
 
+std::vector<std::wstring> GCommandLine;
+
+void
+LogSaveCommandLineUtf8(int Argc, char** Argv)
+{
+	std::lock_guard<std::mutex> LockGuard(GLogMutex);
+	for (int i = 0; i < Argc; ++i)
+	{
+		GCommandLine.push_back(ConvertUtf8ToWide(Argv[i]));
+	}
+}
+
 static FILE*
 GetLogStream(ELogLevel LogLevel)
 {
@@ -142,13 +155,29 @@ std::vector<std::unique_ptr<FLogFile>> GLogFileStack;
 
 FLogFileScope::FLogFileScope(const wchar_t* Filename)
 {
+	std::wstring CommandLine;
+
 	{
 		std::lock_guard<std::mutex> LockGuard(GLogMutex);
 		GLogFileStack.push_back(std::move(GLogFile));
 		LogSetFileInternal(Filename);
+
+		for (const std::wstring& Arg : GCommandLine)
+		{
+			if (!CommandLine.empty())
+			{
+				CommandLine += ' ';
+			}
+			CommandLine += Arg;
+		}
 	}
 
 	UNSYNC_VERBOSE2(L"UNSYNC v%hs started logging to file '%ls'", GetVersionString().c_str(), Filename);
+	
+	if (!CommandLine.empty())
+	{
+		UNSYNC_VERBOSE2(L"Command line: %ls", CommandLine.c_str());
+	}
 }
 
 FLogFileScope::~FLogFileScope()
@@ -267,7 +296,11 @@ LogPrintf(ELogLevel Level, const wchar_t* Str, ...)
 
 	ELogLevel MaxDisplayLevel = ELogLevel::Info;
 
-	if (GLogVerbose)
+	if (GLogSilent)
+	{
+		MaxDisplayLevel = ELogLevel::Warning;
+	}
+	else if (GLogVerbose)
 	{
 		if (GLogVeryVerbose)
 		{
@@ -290,7 +323,7 @@ LogPrintf(ELogLevel Level, const wchar_t* Str, ...)
 
 		if (bShouldOutputThreadIndex)
 		{
-			fwprintf(LogStream, L"[Thread %d] ", ThreadIndex);
+			fwprintf(LogStream, L"[Thread %u] ", ThreadIndex);
 		}
 
 		if (bShouldIndent)
@@ -314,7 +347,7 @@ LogPrintf(ELogLevel Level, const wchar_t* Str, ...)
 
 		fwprintf(LogFileStream, L"[%hs] ", TimestampString.data());
 
-		fwprintf(LogFileStream, L"[%3d] ", ThreadIndex);
+		fwprintf(LogFileStream, L"[%3u] ", ThreadIndex);
 
 		switch (Level)
 		{
diff --git a/Engine/Source/Programs/Unsync/Private/UnsyncLog.h b/Engine/Source/Programs/Unsync/Private/UnsyncLog.h
--- a/Engine/Source/Programs/Unsync/Private/UnsyncLog.h
+++ b/Engine/Source/Programs/Unsync/Private/UnsyncLog.h
@@ -8,6 +8,7 @@ namespace unsync {
 
 extern thread_local bool   GLogVerbose;
 extern bool				   GLogVeryVerbose;
+extern bool				   GLogSilent;
 extern bool				   GBreakOnError;
 extern bool				   GBreakOnWarning;
 extern thread_local uint32 GLogIndent;
@@ -93,6 +94,8 @@ void SetCrashDumpPath(const FPath& Path);
 bool GetCrashDumpPath(FPath& OutPath);
 bool LogWriteCrashDump(void* ExceptionPointers);
 
+void LogSaveCommandLineUtf8(int Argc, char** Argv);
+
 #define UNSYNC_LOG_INDENT		FLogIndentScope UNSYNC_CONCAT(log_indent_, __LINE__);
 #define UNSYNC_LOG_INDENT_N(N)  FLogIndentScope UNSYNC_CONCAT(log_indent_, __LINE__)(N);
 
diff --git a/Engine/Source/Programs/Unsync/Private/UnsyncMain.cpp b/Engine/Source/Programs/Unsync/Private/UnsyncMain.cpp
--- a/Engine/Source/Programs/Unsync/Private/UnsyncMain.cpp
+++ b/Engine/Source/Programs/Unsync/Private/UnsyncMain.cpp
@@ -38,7 +38,7 @@ static FPath GExePath;
 int
 InnerMain(int Argc, char** Argv)
 {
-	FTimingLogger TimingLogger("Total time");
+	LogSaveCommandLineUtf8(Argc, Argv);
 
 	std::string AppDescription = "UNSYNC v";
 	AppDescription += GetVersionString();
@@ -265,7 +265,6 @@ InnerMain(int Argc, char** Argv)
 
 	SubSync->add_flag("--no-output-validation", bNoOutputValidation, "Skip final patched file block hash validation (DANGEROUS)");
 	SubSync->add_flag("--no-space-validation", bNoSpaceValidation, "Skip checking available disk space before sync (DANGEROUS)");
-	SubSync->add_option("-b, --block", HashOrSyncBlockSize, "Block size in bytes (default=64KB)");
 	SubSync->add_option("--scavenge", ScavengeRootUtf8, "Search for unsync manifests and reusable blocks in this directory (EXPERIMENTAL)");
 	SubSync->add_flag("--login", bShouldLogin, "Use user authentication when accessing unsync server");
 	SubSync->add_flag("--no-timeout", bNoSocketTimeout, "Disable the default 60 second timeout on network socket operations");
@@ -320,12 +319,17 @@ InnerMain(int Argc, char** Argv)
 	for (CLI::App* Subcommand : SubCommands)
 	{
 		Subcommand->add_flag("-d, --dry, --dry-run", GDryRun, "Don't write any outputs to disk");
-		Subcommand->add_flag("-v, --verbose", GLogVerbose, "Verbose logging");
-		Subcommand->add_flag("--very-verbose", GLogVeryVerbose, "Very verbose logging");
+		auto VerboseFlag	 = Subcommand->add_flag("-v, --verbose", GLogVerbose, "Verbose logging");
+		auto VeryVerboseFlag = Subcommand->add_flag("--very-verbose", GLogVeryVerbose, "Very verbose logging");
+		auto SilentFlag		 = Subcommand->add_flag("--silent", GLogSilent, "Skip all console logging except errors and warnings");
 		Subcommand->add_flag("--progress", GLogProgress, "Output @progress and @status markers");
 		Subcommand->add_option("--threads", GMaxThreads, "Limit worker threads to specified number");
 		Subcommand->add_flag("--buffered-files", GForceBufferedFiles, "Always use buffered file IO");
 		Subcommand->add_flag("--debug", bUseDebugMode, "Enable extra debugging features, such as extra memory safety validation");
+
+		SilentFlag->excludes(VerboseFlag);
+		SilentFlag->excludes(VeryVerboseFlag);
+		VeryVerboseFlag->excludes(VerboseFlag);
 	}
 
 	// Run the command
@@ -348,6 +352,8 @@ InnerMain(int Argc, char** Argv)
 		wprintf(L"%hs", Cli.help().c_str());
 	}
 
+	FTimingLogger TimingLogger("Total time", ELogLevel::Info);
+
 	// Configure default output mehtod based on subcommand.
 	// In machine-readable mode, all verbose logging is directed to stderr.
 
@@ -369,7 +375,7 @@ InnerMain(int Argc, char** Argv)
 
 	if (const char* EnvCleanupExclude = getenv("UNSYNC_CLEANUP_EXCLUDE"))
 	{
-		UNSYNC_VERBOSE(L"Using UNSYNC_CLEANUP_EXCLUDE environment: '%hs'", EnvCleanupExclude);
+		UNSYNC_LOG(L"Using UNSYNC_CLEANUP_EXCLUDE environment: '%hs'", EnvCleanupExclude);
 		CleanupExcludeFilterArrayUtf8.push_back(EnvCleanupExclude);
 	}
 
@@ -378,7 +384,7 @@ InnerMain(int Argc, char** Argv)
 		const char* EnvDfs = getenv("UNSYNC_DFS");
 		if (EnvDfs)
 		{
-			UNSYNC_VERBOSE(L"Using UNSYNC_DFS environment: '%hs'", EnvDfs);
+			UNSYNC_LOG(L"Using UNSYNC_DFS environment: '%hs'", EnvDfs);
 			PreferredDfsUtf8 = std::string(EnvDfs);
 		}
 	}
@@ -388,7 +394,7 @@ InnerMain(int Argc, char** Argv)
 		const char* EnvProxy = getenv("UNSYNC_PROXY");
 		if (EnvProxy)
 		{
-			UNSYNC_VERBOSE(L"Using UNSYNC_PROXY environment: '%hs'", EnvProxy);
+			UNSYNC_LOG(L"Using UNSYNC_PROXY environment: '%hs'", EnvProxy);
 			RemoteAddressUtf8 = std::string(EnvProxy);
 		}
 	}
@@ -398,7 +404,7 @@ InnerMain(int Argc, char** Argv)
 		const char* EnvCacert = getenv("UNSYNC_CACERT");
 		if (EnvCacert)
 		{
-			UNSYNC_VERBOSE(L"Using UNSYNC_CACERT environment: '%hs'", EnvCacert);
+			UNSYNC_LOG(L"Using UNSYNC_CACERT environment: '%hs'", EnvCacert);
 			CacertFilenameUtf8 = std::string(EnvCacert);
 		}
 	}
@@ -408,7 +414,7 @@ InnerMain(int Argc, char** Argv)
 		const char* EnvHttpHeaderFile = getenv("UNSYNC_HTTP_HEADER_FILE");
 		if (EnvHttpHeaderFile)
 		{
-			UNSYNC_VERBOSE(L"Using UNSYNC_HTTP_HEADER_FILE environment: '%hs'", EnvHttpHeaderFile);
+			UNSYNC_LOG(L"Using UNSYNC_HTTP_HEADER_FILE environment: '%hs'", EnvHttpHeaderFile);
 			HttpHeaderFilenameUtf8 = std::string(EnvHttpHeaderFile);
 		}
 	}
@@ -540,16 +546,16 @@ InnerMain(int Argc, char** Argv)
 
 	if (GLogVeryVerbose)
 	{
-		UNSYNC_VERBOSE(L"Very verbose logging is enabled");
+		UNSYNC_LOG(L"Very verbose logging is enabled");
 	}
 	else if (GLogVerbose)
 	{
-		UNSYNC_VERBOSE(L"Verbose logging is enabled");
+		UNSYNC_LOG(L"Verbose logging is enabled");
 	}
 
 	if (GDryRun)
 	{
-		UNSYNC_VERBOSE(L">>> DRY RUN <<<");
+		UNSYNC_LOG(L">>> DRY RUN <<<");
 	}
 
 	if (GForceBufferedFiles)
@@ -561,14 +567,14 @@ InnerMain(int Argc, char** Argv)
 	UNSYNC_VERBOSE2(L"Using threads: %d", GMaxThreads);
 	FConcurrencyPolicyScope ConcurrencyLimitScope(GMaxThreads);
 
-	if (Cli.got_subcommand(SubHash) || Cli.got_subcommand(SubSync) || Cli.got_subcommand(SubPack))
+	if (Cli.got_subcommand(SubHash) || Cli.got_subcommand(SubPack))
 	{
-		UNSYNC_VERBOSE(L"Using block size: %d KB", HashOrSyncBlockSize / 1024);
+		UNSYNC_LOG(L"Using block size: %d KB", HashOrSyncBlockSize / 1024);
 	}
 
 	if (Cli.got_subcommand(SubDiff))
 	{
-		UNSYNC_VERBOSE(L"Using block size: %d KB", DiffBlockSize / 1024);
+		UNSYNC_LOG(L"Using block size: %d KB", DiffBlockSize / 1024);
 	}
 
 	if (Cli.got_subcommand(SubDiff))
@@ -611,7 +617,7 @@ InnerMain(int Argc, char** Argv)
 	if (!PreferredDfsUtf8.empty() && !SourceFilenameUtf8.empty())
 	{
 		LogGlobalStatus(L"Enumerating DFS");
-		UNSYNC_VERBOSE(L"Enumerating DFS");
+		UNSYNC_LOG(L"Enumerating DFS");
 		std::wstring PreferredDfs = ConvertUtf8ToWide(PreferredDfsUtf8);
 		auto		 DfsEntries	  = DfsEnumerate(SourceFilename);
 
@@ -629,7 +635,7 @@ InnerMain(int Argc, char** Argv)
 
 		if (FoundDfsStorage)
 		{
-			UNSYNC_VERBOSE(L"Found preferred DFS storage server '%ls' with share '%ls'",
+			UNSYNC_LOG(L"Found preferred DFS storage server '%ls' with share '%ls'",
 						   FoundDfsStorage->Server.c_str(),
 						   FoundDfsStorage->Share.c_str());
 
@@ -637,7 +643,7 @@ InnerMain(int Argc, char** Argv)
 			DfsAlias.Source = DfsEntries.Root;
 			DfsAlias.Target = FPath(L"\\\\") / FoundDfsStorage->Server / FoundDfsStorage->Share;
 
-			UNSYNC_VERBOSE(L"Using DFS alias '%ls' -> '%ls'", DfsAlias.Source.wstring().c_str(), DfsAlias.Target.wstring().c_str());
+			UNSYNC_LOG(L"Using DFS alias '%ls' -> '%ls'", DfsAlias.Source.wstring().c_str(), DfsAlias.Target.wstring().c_str());
 
 			if (!DfsAlias.Source.empty())
 			{
@@ -788,12 +794,12 @@ InnerMain(int Argc, char** Argv)
 
 		if (bShouldLogin)
 		{
-			UNSYNC_VERBOSE(L"Attempting to authenticate");
+			UNSYNC_LOG(L"Attempting to authenticate");
 			UNSYNC_LOG_INDENT;
 			bool bUsingAuthentication = TryAddAuthentication(RemoteDesc);
 			if (bUsingAuthentication)
 			{
-				UNSYNC_VERBOSE(L"Authentication enabled");
+				UNSYNC_LOG(L"Authentication enabled");
 			}
 		}
 
@@ -816,7 +822,6 @@ InnerMain(int Argc, char** Argv)
 		SyncOptions.bFullDifference		   = bFullDifference;
 		SyncOptions.bFullSourceScan		   = bFullSourceScan;
 		SyncOptions.bCleanup			   = !bNoCleanupAfterSync;
-		SyncOptions.BlockSize			   = HashOrSyncBlockSize;
 		SyncOptions.Filter				   = &SyncFilter;
 		SyncOptions.bValidateTargetFiles   = !bNoOutputValidation;
 		SyncOptions.bCheckAvailableSpace   = !bNoSpaceValidation;
diff --git a/Engine/Source/Programs/Unsync/Private/UnsyncUtil.cpp b/Engine/Source/Programs/Unsync/Private/UnsyncUtil.cpp
--- a/Engine/Source/Programs/Unsync/Private/UnsyncUtil.cpp
+++ b/Engine/Source/Programs/Unsync/Private/UnsyncUtil.cpp
@@ -63,12 +63,12 @@ BytesToHexString(const uint8* Data, uint64 Size)
 	return Result;
 }
 
-FTimingLogger::FTimingLogger(const char* InName, bool bInEnabled) : bEnabled(bInEnabled), Name(InName)
+FTimingLogger::FTimingLogger(const char* InName, ELogLevel InLogLevel, bool bInEnabled)
+: bEnabled(bInEnabled)
+, Name(InName)
+, LogLevel(InLogLevel)
 {
-	if (bEnabled)
-	{
-		TimeBegin = TimePointNow();
-	}
+	TimeBegin = TimePointNow();
 }
 
 FTimingLogger::~FTimingLogger()
@@ -90,13 +90,20 @@ void FTimingLogger::Finish()
 
 		if (Name.empty())
 		{
-			UNSYNC_VERBOSE(L"%.2f sec (%02d:%02d:%02d)", TotalSeconds, H, M, S);
+			LogPrintf(LogLevel, L"%.3f sec", TotalSeconds);
 		}
 		else
 		{
-			UNSYNC_VERBOSE(L"%hs: %.2f sec (%02d:%02d:%02d)", Name.c_str(), TotalSeconds, H, M, S);
+			LogPrintf(LogLevel, L"%hs: %.3f sec", Name.c_str(), TotalSeconds);
 		}
 
+		if (TotalSeconds >= 60.0)
+		{
+			LogPrintf(LogLevel, L" (%02d:%02d:%02d)", H, M, S);
+		}
+
+		LogPrintf(LogLevel, L"\n");
+
 		LogFlush();
 
 		bEnabled = false;
diff --git a/Engine/Source/Programs/Unsync/Private/UnsyncUtil.h b/Engine/Source/Programs/Unsync/Private/UnsyncUtil.h
--- a/Engine/Source/Programs/Unsync/Private/UnsyncUtil.h
+++ b/Engine/Source/Programs/Unsync/Private/UnsyncUtil.h
@@ -85,8 +85,9 @@ struct FTimingLogger
 	bool		bEnabled  = false;
 	FTimePoint	TimeBegin = FTimePoint{};
 	std::string Name;
+	ELogLevel	LogLevel;
 
-	FTimingLogger(const char* InName, bool InEnabled = true);
+	FTimingLogger(const char* InName, ELogLevel InLogLevel, bool InEnabled = true);
 	~FTimingLogger();
 	void Finish();
 };
