unsync - Reuse existing block database during pack generation

diff --git a/Engine/Source/Programs/Unsync/Private/UnsyncCmdPack.cpp b/Engine/Source/Programs/Unsync/Private/UnsyncCmdPack.cpp
--- a/Engine/Source/Programs/Unsync/Private/UnsyncCmdPack.cpp
+++ b/Engine/Source/Programs/Unsync/Private/UnsyncCmdPack.cpp
@@ -83,6 +83,14 @@ BuildP4HaveSet(const FPath& Root, std::string_view P4HaveDataUtf8, FDirectoryMan
 	ForLines(P4HaveDataUtf8, Callback);
 }
 
+static void DropBlocksFromManifest(FDirectoryManifest& DirectoryManifest)
+{
+	for (auto& FileEntry : DirectoryManifest.Files)
+	{
+		FileEntry.second.Blocks.clear();
+	}
+}
+
 struct FPackIndexEntry
 {
 	FHash160 Hash	 = {};
@@ -121,6 +129,50 @@ int32 CmdPack(const FCmdPackOptions& Options)
 		return -1;
 	}
 
+	FPath		TempOutputPackFilename = ManifestRoot / "blocks.bin.tmp";
+	FNativeFile PackFile(TempOutputPackFilename, EFileMode::CreateWriteOnly);
+	if (!PackFile.IsValid())
+	{
+		UNSYNC_ERROR(L"Failed open pack output file '%ls'", TempOutputPackFilename.wstring().c_str());
+		return -1;
+	}
+
+	FPath		TempOutputIndexFilename = ManifestRoot / "blocks.idx.tmp";
+	FNativeFile IndexFile(TempOutputIndexFilename, EFileMode::CreateWriteOnly);
+	if (!IndexFile.IsValid())
+	{
+		UNSYNC_ERROR(L"Failed open index output file '%ls'", TempOutputIndexFilename.wstring().c_str());
+		return -1;
+	}
+
+	THashSet<FHash160> SeenMacroBlockHashSet;
+	UNSYNC_LOG(L"Loading existing block packs ...");
+	{
+		UNSYNC_LOG_INDENT;
+		FPath ExistingIndexSearchPath = ManifestRoot; // TODO: could allow overriding this
+		FPath ExpectedExtension		  = FPath(".idx");
+		for (const std::filesystem::directory_entry& Dir : RecursiveDirectoryScan(ExistingIndexSearchPath))
+		{
+			if (!Dir.is_regular_file())
+			{
+				continue;
+			}
+
+			const FPath& FilePath = Dir.path();
+			FPathStringView FilePathView(FilePath.native());
+			if (!FilePathView.ends_with(ExpectedExtension.native()))
+			{
+				continue;
+			}
+
+			FBuffer ExistingEntries = ReadFileToBuffer(FilePath);
+			uint64	NumEntries		= ExistingEntries.Size() / sizeof(FPackIndexEntry);
+			for (const FPackIndexEntry& Entry : MakeView(reinterpret_cast<FPackIndexEntry*>(ExistingEntries.Data()), NumEntries))
+			{
+				SeenMacroBlockHashSet.insert(Entry.Hash);
+			}
+		}
+	}
 
 	const FPath DirectoryManifestPath = ManifestRoot / "manifest.bin";
 
@@ -177,14 +229,6 @@ int32 CmdPack(const FCmdPackOptions& Options)
 
 	UNSYNC_LOG(L"Found files: %llu", llu(DirectoryManifest.Files.size()));
 
-	FPath OutputPackFilename = ManifestRoot / "blocks.bin";
-	FNativeFile PackFile(OutputPackFilename, EFileMode::CreateWriteOnly);
-
-	FPath		OutputIndexFilename = ManifestRoot / "blocks.idx";
-	FNativeFile IndexFile(OutputIndexFilename, EFileMode::CreateWriteOnly);
-
-	UNSYNC_LOG(L"Building file blocks ...");
-
 	FComputeBlocksParams BlockParams;
 	BlockParams.Algorithm = Options.Algorithm;
 	BlockParams.BlockSize = Options.BlockSize;
@@ -192,17 +236,19 @@ int32 CmdPack(const FCmdPackOptions& Options)
 
 	std::atomic<uint64> PackFileOffset;
 	std::atomic<uint64> IndexFileOffset;
+	std::atomic<uint64> ProcessedRawBytes;
 
-	THashSet<FGenericHash> SeenMacroBlockHashSet;
+	std::atomic<uint64> IndexFileHashSum[2];
 
 	std::mutex Mutex;
 	BlockParams.OnMacroBlockGenerated =
-		[&PackFile, &PackFileOffset, &IndexFile, &IndexFileOffset, &SeenMacroBlockHashSet, &Mutex](const FGenericBlock& Block,
-																								   FBufferView			Data)
+		[&Mutex, &PackFile, &PackFileOffset, &IndexFile, &IndexFileOffset, &SeenMacroBlockHashSet, &IndexFileHashSum, &ProcessedRawBytes](
+			const FGenericBlock& Block,
+			FBufferView			 Data)
 	{
 		{
 			std::lock_guard<std::mutex> LockGuard(Mutex);
-			if (!SeenMacroBlockHashSet.insert(Block.HashStrong).second)
+			if (!SeenMacroBlockHashSet.insert(Block.HashStrong.ToHash160()).second)
 			{
 				return;
 			}
@@ -214,6 +260,8 @@ int32 CmdPack(const FCmdPackOptions& Options)
 
 		if (ActualCompressedSize)
 		{
+			ProcessedRawBytes += Data.Size;
+
 			CompressedData.SetDataRange(0, ActualCompressedSize);
 
 			uint64 PackWriteOffset = PackFileOffset.fetch_add(ActualCompressedSize);
@@ -221,14 +269,55 @@ int32 CmdPack(const FCmdPackOptions& Options)
 
 			FPackIndexEntry IndexEntry;
 			IndexEntry.Hash			  = Block.HashStrong.ToHash160();
-			IndexEntry.CompressedSize = CheckedNarrow<uint32>(ActualCompressedSize);
+			IndexEntry.CompressedSize = CheckedNarrow<uint64>(ActualCompressedSize);
 			IndexEntry.Offset		  = PackWriteOffset;
 
 			uint64 IndexWriteOffset = IndexFileOffset.fetch_add(sizeof(IndexEntry));
 			IndexFile.Write(&IndexEntry, IndexWriteOffset, sizeof(IndexEntry));
+
+			FHash128 BlockHash128 = Block.HashStrong.ToHash128();
+			uint64	 BlockHashParts[2];
+			memcpy(BlockHashParts, &BlockHash128, sizeof(BlockHash128));
+			IndexFileHashSum[0] += BlockHashParts[0];
+			IndexFileHashSum[1] += BlockHashParts[1];
 		}
 	};
 
+	{
+		UNSYNC_LOG(L"Loading previous manifest ... ")
+		FDirectoryManifest OldManifest;
+		if (LoadDirectoryManifest(OldManifest, InputRoot, DirectoryManifestPath))
+		{
+			UNSYNC_LOG(L"Previous manifest loaded")
+
+			BlockParams.Algorithm = OldManifest.Algorithm;
+
+			FDirectoryManifest& NewManifest = DirectoryManifest;
+
+			// Copy file blocks from old manifest, if possible
+			for (auto& NewManifestFileEntry : NewManifest.Files)
+			{
+				const std::wstring& FileName			 = NewManifestFileEntry.first;
+				auto				OldManifestFileEntry = OldManifest.Files.find(FileName);
+				if (OldManifestFileEntry == OldManifest.Files.end())
+				{
+					continue;
+				}
+
+				FFileManifest& NewEntry = NewManifestFileEntry.second;
+				FFileManifest& OldEntry = OldManifestFileEntry->second;
+
+				if (NewEntry.Mtime == OldEntry.Mtime && NewEntry.Size == OldEntry.Size)
+				{
+					NewEntry.Blocks		 = std::move(OldEntry.Blocks);
+					NewEntry.MacroBlocks = std::move(OldEntry.MacroBlocks);
+					NewEntry.BlockSize	 = OldEntry.BlockSize;
+				}
+			}
+		}
+	}
+
+	UNSYNC_LOG(L"Computing file blocks ...");
 	UpdateDirectoryManifestBlocks(DirectoryManifest, InputRoot, BlockParams);
 
 	const uint64 CompressedSize = PackFileOffset;
@@ -236,6 +325,7 @@ int32 CmdPack(const FCmdPackOptions& Options)
 	if (!GDryRun)
 	{
 		UNSYNC_LOG(L"Saving directory manifest '%ls'", DirectoryManifestPath.wstring().c_str());
+		DropBlocksFromManifest(DirectoryManifest); // only need to keep macro blocks
 		SaveDirectoryManifest(DirectoryManifest, DirectoryManifestPath);
 	}
 
@@ -248,12 +338,47 @@ int32 CmdPack(const FCmdPackOptions& Options)
 	const uint64 NumSourceFiles = DirectoryManifest.Files.size();
 	UNSYNC_LOG(L"Source files: %llu", llu(NumSourceFiles));
 	UNSYNC_LOG(L"Source size: %llu bytes (%.2f MB)", llu(SourceSize), SizeMb(SourceSize));
+	UNSYNC_LOG(L"New data size: %llu bytes (%.2f MB)", llu(ProcessedRawBytes), SizeMb(ProcessedRawBytes));
 	UNSYNC_LOG(L"Compressed size: %llu bytes (%.2f MB), %.0f%%",
-				   llu(CompressedSize),
-				   SizeMb(CompressedSize),
-				   100.0 * double(CompressedSize) / double(SourceSize));
+			   llu(CompressedSize),
+			   SizeMb(CompressedSize),
+			   ProcessedRawBytes > 0 ? (100.0 * double(CompressedSize) / double(ProcessedRawBytes)) : 0);
 
 	PackFile.Close();
+	IndexFile.Close();
+
+	std::error_code ErrorCode;
+
+	if (PackFileOffset == 0)
+	{
+		UNSYNC_LOG(L"No new blocks found");
+		FileRemove(TempOutputPackFilename, ErrorCode);
+		FileRemove(TempOutputIndexFilename, ErrorCode);
+	}
+	else
+	{
+		uint64 IndexHashParts[2];
+		IndexHashParts[0] = IndexFileHashSum[0];
+		IndexHashParts[1] = IndexFileHashSum[1];
+		FHash128 BlockHash128;
+		memcpy(&BlockHash128, IndexHashParts, sizeof(BlockHash128));
+		std::string OutputId = HashToHexString(BlockHash128);
+
+		FPath FinalPackFilename	 = ManifestRoot / (OutputId + ".bin");
+		FPath FinalIndexFilename = ManifestRoot / (OutputId + ".idx");
+
+		if (!FileRename(TempOutputPackFilename, FinalPackFilename, ErrorCode))
+		{
+			UNSYNC_ERROR(L"Failed to rename temporary pack file to '%ls'", FinalPackFilename.wstring().c_str());
+			return 1;
+		}
+
+		if (!FileRename(TempOutputIndexFilename, FinalIndexFilename, ErrorCode))
+		{
+			UNSYNC_ERROR(L"Failed to rename temporary index file to '%ls'", FinalIndexFilename.wstring().c_str());
+			return 1;
+		}
+	}
 
 	return 0;
 }
diff --git a/Engine/Source/Programs/Unsync/Private/UnsyncCmdSync.cpp b/Engine/Source/Programs/Unsync/Private/UnsyncCmdSync.cpp
--- a/Engine/Source/Programs/Unsync/Private/UnsyncCmdSync.cpp
+++ b/Engine/Source/Programs/Unsync/Private/UnsyncCmdSync.cpp
@@ -189,7 +189,7 @@ CmdSync(const FCmdSyncOptions& Options)
 
 			FSyncFileOptions SyncFileOptions;
 			SyncFileOptions.Algorithm			 = Options.Algorithm;
-			SyncFileOptions.BlockSize			 = 64_KB;
+			SyncFileOptions.BlockSize			 = uint32(64_KB);
 			SyncFileOptions.bValidateTargetFiles = Options.bValidateTargetFiles;
 
 			return SyncFile(Options.Source, Options.Target, Options.Target, SyncFileOptions).Succeeded() ? 0 : 1;
diff --git a/Engine/Source/Programs/Unsync/Private/UnsyncCore.cpp b/Engine/Source/Programs/Unsync/Private/UnsyncCore.cpp
--- a/Engine/Source/Programs/Unsync/Private/UnsyncCore.cpp
+++ b/Engine/Source/Programs/Unsync/Private/UnsyncCore.cpp
@@ -2019,7 +2019,9 @@ UpdateDirectoryManifestBlocks(FDirectoryManifest& Result, const FPath& Root, con
 	for (auto& It : Result.Files)
 	{
 		FFileManifest& FileManifest = It.second;
-		if (FileManifest.BlockSize == Params.BlockSize)
+		const bool	   bBlockSizeOk	  = FileManifest.BlockSize == Params.BlockSize;
+		const bool	   bMacroBlocksOk = Params.bNeedMacroBlocks ? !FileManifest.MacroBlocks.empty() : true;
+		if (bBlockSizeOk && bMacroBlocksOk)
 		{
 			NumSkippedBlocks += FileManifest.Blocks.size();
 			NumSkippedBytes += FileManifest.Size;
@@ -2226,7 +2228,7 @@ ComputeManifestStableSignature(const FDirectoryManifest& Manifest)
 	UpdateHashT(Hasher, Manifest.Algorithm.WeakHashAlgorithmId);
 	UpdateHashT(Hasher, Manifest.Algorithm.StrongHashAlgorithmId);
 
-	std::vector<std::wstring> SortedFiles;
+	std::vector<std::wstring_view> SortedFiles;
 	SortedFiles.reserve(Manifest.Files.size());
 
 	for (const auto& It : Manifest.Files)
@@ -2238,13 +2240,17 @@ ComputeManifestStableSignature(const FDirectoryManifest& Manifest)
 	std::sort(SortedFiles.begin(), SortedFiles.end());
 
 	std::string FileNameUtf8;
-	for (const std::wstring& FileName : SortedFiles)
+	std::wstring FileName;
+	for (const std::wstring_view& FileNameView : SortedFiles)
 	{
+		FileName.clear();
+		FileName.append(FileNameView);
+
 		// Canonical unsync file paths are utf8 with unix-style separator `/`
 		ConvertWideToUtf8(FileName, FileNameUtf8);
 		std::replace(FileNameUtf8.begin(), FileNameUtf8.end(), '\\', '/');
 
-		const FFileManifest& FileManifest = Manifest.Files.at(FileName);
+		const FFileManifest& FileManifest = Manifest.Files.at(FileName);  // TODO: heterogeneous map
 
 		blake3_hasher_update(&Hasher, FileNameUtf8.c_str(), FileNameUtf8.length());
 
diff --git a/Engine/Source/Programs/Unsync/Private/UnsyncCore.h b/Engine/Source/Programs/Unsync/Private/UnsyncCore.h
--- a/Engine/Source/Programs/Unsync/Private/UnsyncCore.h
+++ b/Engine/Source/Programs/Unsync/Private/UnsyncCore.h
@@ -194,7 +194,7 @@ using FOnBlockGenerated = std::function<void(const FGenericBlock& Block, FBuffer
 struct FComputeBlocksParams
 {
 	bool			  bNeedBlocks = true;
-	uint32			  BlockSize	  = 64_KB;
+	uint32			  BlockSize	  = uint32(64_KB);
 	FAlgorithmOptions Algorithm;
 
 	bool   bNeedMacroBlocks		= false;
