unsync - Move most of the manifest related utility code to a dedicated file

diff --git a/Engine/Source/Programs/Unsync/CMakeLists.txt b/Engine/Source/Programs/Unsync/CMakeLists.txt
--- a/Engine/Source/Programs/Unsync/CMakeLists.txt
+++ b/Engine/Source/Programs/Unsync/CMakeLists.txt
@@ -42,6 +42,8 @@ set(src
 	Private/UnsyncCmdLogin.h
 	Private/UnsyncCmdMount.cpp
 	Private/UnsyncCmdMount.h
+	Private/UnsyncCmdPack.cpp
+	Private/UnsyncCmdPack.h
 	Private/UnsyncCmdPatch.cpp
 	Private/UnsyncCmdPatch.h
 	Private/UnsyncCmdPush.cpp
@@ -50,8 +52,6 @@ set(src
 	Private/UnsyncCmdQuery.h
 	Private/UnsyncCmdSync.cpp
 	Private/UnsyncCmdSync.h
-	Private/UnsyncCmdPack.cpp
-	Private/UnsyncCmdPack.h
 	Private/UnsyncCommon.h
 	Private/UnsyncCompression.cpp
 	Private/UnsyncCompression.h
@@ -70,6 +70,8 @@ set(src
 	Private/UnsyncLog.cpp
 	Private/UnsyncLog.h
 	Private/UnsyncMain.cpp
+	Private/UnsyncManifest.cpp
+	Private/UnsyncManifest.h
 	Private/UnsyncMemory.cpp
 	Private/UnsyncMemory.h
 	Private/UnsyncMiniCb.cpp
diff --git a/Engine/Source/Programs/Unsync/Private/UnsyncCmdPack.cpp b/Engine/Source/Programs/Unsync/Private/UnsyncCmdPack.cpp
--- a/Engine/Source/Programs/Unsync/Private/UnsyncCmdPack.cpp
+++ b/Engine/Source/Programs/Unsync/Private/UnsyncCmdPack.cpp
@@ -1,7 +1,5 @@
 // Copyright Epic Games, Inc. All Rights Reserved.
 
-#pragma once
-
 #include "UnsyncCmdPack.h"
 #include "UnsyncFile.h"
 #include "UnsyncThread.h"
diff --git a/Engine/Source/Programs/Unsync/Private/UnsyncCore.cpp b/Engine/Source/Programs/Unsync/Private/UnsyncCore.cpp
--- a/Engine/Source/Programs/Unsync/Private/UnsyncCore.cpp
+++ b/Engine/Source/Programs/Unsync/Private/UnsyncCore.cpp
@@ -1833,447 +1833,6 @@ IsSynchronized(const FNeedList& NeedList, const FGenericBlockArray& SourceBlocks
 	return true;
 }
 
-FDirectoryManifest
-CreateDirectoryManifest(const FPath& Root, const FComputeBlocksParams& Params)
-{
-	UNSYNC_LOG_INDENT;
-
-	FDirectoryManifest Result;
-
-	Result.Algorithm = Params.Algorithm;
-
-	FTimePoint TimeBegin = TimePointNow();
-
-	FTaskGroup	 TaskGroup;
-	const uint32 MaxConcurrentFiles = 8;  // quickly diminishing returns past 8 concurrent files
-	FSemaphore	 Semaphore(MaxConcurrentFiles);
-
-	std::mutex ResultMutex;
-	FPath	   UnsyncDirName = ".unsync";
-
-	for (const std::filesystem::directory_entry& Dir : RecursiveDirectoryScan(Root))
-	{
-		if (Dir.is_directory())
-		{
-			continue;
-		}
-
-		FPath RelativePath = GetRelativePath(Dir.path(), Root);
-
-		if (RelativePath.native().starts_with(UnsyncDirName.native()))
-		{
-			continue;
-		}
-
-		const std::wstring PathKey = RelativePath.wstring();
-
-		UNSYNC_VERBOSE2(L"Found '%ls'", PathKey.c_str());
-
-		FFileManifest FileManifest;
-
-		FileManifest.Mtime		 = ToWindowsFileTime(Dir.last_write_time());
-		FileManifest.Size		 = Dir.file_size();
-		FileManifest.CurrentPath = Dir.path();
-		FileManifest.BlockSize	 = Params.BlockSize;
-		FileManifest.bReadOnly	 = IsReadOnly(Dir.status().permissions());
-
-		{
-			std::lock_guard<std::mutex> LockGuard(ResultMutex);
-			Result.Files[PathKey] = std::move(FileManifest);
-		}
-
-		if (Params.bNeedBlocks && Params.BlockSize)
-		{
-			FPath FilePath = Root / RelativePath;
-			auto  File	   = std::make_shared<FNativeFile>(FilePath, EFileMode::ReadOnlyUnbuffered);
-			if (File->IsValid())
-			{
-				UNSYNC_VERBOSE(L"Computing blocks for '%ls' (%.2f MB)", FilePath.wstring().c_str(), double(File->GetSize()) / (1 << 20));
-
-				Semaphore.Acquire();
-				TaskGroup.run([&Semaphore, &ResultMutex, &Result, File = std::move(File), Key = std::move(PathKey), &Params]() {
-
-					FComputeBlocksResult ComputedBlocks = ComputeBlocks(*File, Params);
-
-					std::lock_guard<std::mutex> LockGuard(ResultMutex);
-
-					std::swap(Result.Files[Key].Blocks, ComputedBlocks.Blocks);
-					std::swap(Result.Files[Key].MacroBlocks, ComputedBlocks.MacroBlocks);
-
-					Semaphore.Release();
-				});
-			}
-			else
-			{
-				UNSYNC_FATAL(L"Failed to open file '%ls' while computing manifest blocks. %hs",
-							 FilePath.wstring().c_str(),
-							 FormatSystemErrorMessage(File->GetError()).c_str());
-			}
-		}
-	}
-
-	TaskGroup.wait();
-
-	if (Params.bNeedBlocks && Params.BlockSize)
-	{
-		// TODO: move manifest stat printing into a helper
-
-		uint64 NumTotalMacroBlocks = 0;
-		uint64 NumTotalBlocks	   = 0;
-		uint64 NumTotalBytes	   = 0;
-
-		for (auto& It : Result.Files)
-		{
-			FFileManifest& FileManifest = It.second;
-			NumTotalMacroBlocks += FileManifest.MacroBlocks.size();
-			NumTotalBlocks += FileManifest.Blocks.size();
-			NumTotalBytes += FileManifest.Size;
-		}
-
-		double Duration = DurationSec(TimeBegin, TimePointNow());
-		UNSYNC_VERBOSE(L"Total computed %lld block(s) in %.3f sec (%.3f MB, %.3f MB / sec)",
-					   (long long)NumTotalBlocks,
-					   Duration,
-					   SizeMb(NumTotalBytes),
-					   SizeMb((double(NumTotalBytes) / Duration)));
-		if (NumTotalMacroBlocks)
-		{
-			UNSYNC_VERBOSE(L"Macro blocks: %lld", NumTotalMacroBlocks);
-		}
-	}
-
-	return Result;
-}
-
-FDirectoryManifest
-CreateDirectoryManifestIncremental(const FPath& Root, const FComputeBlocksParams& InParams)
-{
-	FComputeBlocksParams Params = InParams;
-
-	FPath ManifestRoot			= Root / ".unsync";
-	FPath DirectoryManifestPath = ManifestRoot / "manifest.bin";
-
-	FDirectoryManifest OldManifest;
-	const bool		   bExistingManifestLoaded = LoadDirectoryManifest(OldManifest, Root, DirectoryManifestPath);
-
-	// Inherit algorithm options from the existing manifest
-	if (bExistingManifestLoaded)
-	{
-		Params.Algorithm = OldManifest.Algorithm;
-	}
-
-	// Scan the input directory and gather file metadata, without generating blocks
-	FComputeBlocksParams LightweightManifestParams = Params; 
-	LightweightManifestParams.bNeedBlocks = false;
-	LightweightManifestParams.BlockSize = 0;
-
-	FDirectoryManifest NewManifest = CreateDirectoryManifest(Root, LightweightManifestParams);
-
-	// Copy file blocks from old manifest, if possible
-	for (auto& NewManifestFileEntry : NewManifest.Files)
-	{
-		const std::wstring& FileName			 = NewManifestFileEntry.first;
-		auto				OldManifestFileEntry = OldManifest.Files.find(FileName);
-		if (OldManifestFileEntry == OldManifest.Files.end())
-		{
-			continue;
-		}
-
-		FFileManifest& NewEntry = NewManifestFileEntry.second;
-		FFileManifest& OldEntry = OldManifestFileEntry->second;
-
-		if (NewEntry.Mtime == OldEntry.Mtime && NewEntry.Size == OldEntry.Size)
-		{
-			NewEntry.Blocks		 = std::move(OldEntry.Blocks);
-			NewEntry.MacroBlocks = std::move(OldEntry.MacroBlocks);
-			NewEntry.BlockSize	 = OldEntry.BlockSize;
-		}
-	}
-
-	// Generate blocks for changed or new files
-	UpdateDirectoryManifestBlocks(NewManifest, Root, Params);
-
-	return NewManifest;
-}
-
-void
-UpdateDirectoryManifestBlocks(FDirectoryManifest& Result, const FPath& Root, const FComputeBlocksParams& Params)
-{
-	UNSYNC_LOG_INDENT;
-
-	UNSYNC_ASSERT(Params.bNeedBlocks);
-	UNSYNC_ASSERT(Params.BlockSize != 0);
-
-	FTimePoint TimeBegin = TimePointNow();
-
-	uint32 NumProcessedFiles = 0;
-
-	FTaskGroup TaskGroup;
-
-	const uint32 MaxConcurrentFiles = 8;  // quickly diminishing returns past 8 concurrent files
-	FSemaphore	 Semaphore(MaxConcurrentFiles);
-
-	uint64 NumSkippedBlocks = 0;
-	uint64 NumSkippedBytes	= 0;
-
-	for (auto& It : Result.Files)
-	{
-		FFileManifest& FileManifest = It.second;
-		const bool	   bBlockSizeOk	  = FileManifest.BlockSize == Params.BlockSize;
-		const bool	   bMacroBlocksOk = Params.bNeedMacroBlocks ? !FileManifest.MacroBlocks.empty() : true;
-		if (bBlockSizeOk && bMacroBlocksOk)
-		{
-			NumSkippedBlocks += FileManifest.Blocks.size();
-			NumSkippedBytes += FileManifest.Size;
-
-			continue;
-		}
-
-		++NumProcessedFiles;
-
-		FPath FilePath = Root / It.first;
-
-		Semaphore.Acquire();
-
-		UNSYNC_VERBOSE(L"Computing blocks for '%ls' (%.2f MB)", FilePath.wstring().c_str(), SizeMb(It.second.Size));
-		auto BlockTask = [&FileManifest, &Semaphore, &Params, FilePath = std::move(FilePath)]()
-		{
-			FNativeFile File(FilePath, EFileMode::ReadOnlyUnbuffered);
-			if (File.IsValid())
-			{
-				FComputeBlocksResult ComputedBlocks = ComputeBlocks(File, Params);
-				std::swap(FileManifest.Blocks, ComputedBlocks.Blocks);
-				std::swap(FileManifest.MacroBlocks, ComputedBlocks.MacroBlocks);
-
-				FileManifest.BlockSize = Params.BlockSize;
-			}
-			else
-			{
-				UNSYNC_FATAL(L"Failed to open file '%ls' while computing manifest blocks. %hs",
-							 FilePath.wstring().c_str(),
-							 FormatSystemErrorMessage(File.GetError()).c_str());
-			}
-			Semaphore.Release();
-		};
-
-		TaskGroup.run(BlockTask);
-	}
-
-	TaskGroup.wait();
-
-	// TODO: move manifest stat printing into a helper
-
-	uint64 NumTotalMacroBlocks = 0;
-	uint64 NumTotalBlocks	   = 0;
-	uint64 NumTotalBytes	   = 0;
-
-	for (auto& It : Result.Files)
-	{
-		FFileManifest& FileManifest = It.second;
-		NumTotalMacroBlocks += FileManifest.MacroBlocks.size();
-		NumTotalBlocks += FileManifest.Blocks.size();
-		NumTotalBytes += FileManifest.Size;
-	}
-
-	NumTotalBlocks -= NumSkippedBlocks;
-	NumTotalBytes -= NumSkippedBytes;
-
-	double Duration = DurationSec(TimeBegin, TimePointNow());
-
-	if (NumTotalBlocks == 0)
-	{
-		UNSYNC_VERBOSE(L"No blocks needed to be computed")
-	}
-	else
-	{
-		UNSYNC_VERBOSE(L"Total computed %lld block(s) in %.3f sec (%.3f MB, %.3f MB / sec)",
-					   (long long)NumTotalBlocks,
-					   Duration,
-					   SizeMb(NumTotalBytes),
-					   SizeMb((double(NumTotalBytes) / Duration)));
-
-		if (NumTotalMacroBlocks)
-		{
-			UNSYNC_VERBOSE(L"Macro blocks: %lld", NumTotalMacroBlocks);
-		}
-	}
-}
-
-static bool
-AlgorithmOptionsCompatible(const FAlgorithmOptions& A, const FAlgorithmOptions& B)
-{
-	return A.StrongHashAlgorithmId == B.StrongHashAlgorithmId && B.WeakHashAlgorithmId == B.WeakHashAlgorithmId &&
-		   B.ChunkingAlgorithmId == B.ChunkingAlgorithmId;
-}
-
-bool
-LoadOrCreateDirectoryManifest(FDirectoryManifest& Result, const FPath& Root, const FComputeBlocksParams& Params)
-{
-	UNSYNC_LOG_INDENT;
-
-	FPath ManifestRoot			= Root / ".unsync";
-	FPath DirectoryManifestPath = ManifestRoot / "manifest.bin";
-
-	FDirectoryManifest OldDirectoryManifest;
-	FDirectoryManifest NewDirectoryManifest;
-
-	const bool bManifestFileExists = PathExists(DirectoryManifestPath);
-	if (!bManifestFileExists)
-	{
-		UNSYNC_VERBOSE(L"Manifest file '%ls' does not exist", DirectoryManifestPath.wstring().c_str());
-	}
-
-	const bool bExistingManifestLoaded = bManifestFileExists && LoadDirectoryManifest(OldDirectoryManifest, Root, DirectoryManifestPath);
-	const bool bExistingManifestCompatible = AlgorithmOptionsCompatible(OldDirectoryManifest.Algorithm, Params.Algorithm);
-
-	if (bExistingManifestLoaded && bExistingManifestCompatible)
-	{
-		UNSYNC_VERBOSE(L"Loaded existing manifest from '%ls'", DirectoryManifestPath.wstring().c_str());
-
-		UNSYNC_VERBOSE(L"Creating lightweight manifest for '%ls'", Root.wstring().c_str());
-		{
-			// Verify that manifests match in dry run mode.
-			// Otherwise just do a quick manifest generation, without file blocks.
-			FComputeBlocksParams NewParams = Params;
-			if (GDryRun)
-			{
-				NewParams.bNeedBlocks = true;
-			}
-			else
-			{
-				NewParams.bNeedBlocks = false;
-				NewParams.BlockSize	  = 0;
-			}
-
-			NewDirectoryManifest = CreateDirectoryManifest(Root, NewParams);
-		}
-
-		UNSYNC_VERBOSE(L"Comparing manifests");
-		for (const auto& OldManifestIt : OldDirectoryManifest.Files)
-		{
-			auto NewManifestIt = NewDirectoryManifest.Files.find(OldManifestIt.first);
-			if (NewManifestIt != NewDirectoryManifest.Files.end())
-			{
-				const FFileManifest& OldFileManifest = OldManifestIt.second;
-				FFileManifest&		 NewFileManifest = NewManifestIt->second;
-				if (NewFileManifest.Size == OldFileManifest.Size && NewFileManifest.Mtime == OldFileManifest.Mtime)
-				{
-					if (NewFileManifest.BlockSize)
-					{
-						UNSYNC_ASSERT(NewFileManifest.BlockSize == OldFileManifest.BlockSize);
-						UNSYNC_ASSERT(NewFileManifest.Blocks.size() == OldFileManifest.Blocks.size());
-						for (uint64 I = 0; I < NewFileManifest.Blocks.size(); ++I)
-						{
-							const FGenericBlock& NewBlock = NewFileManifest.Blocks[I];
-							const FGenericBlock& OldBlock = OldFileManifest.Blocks[I];
-							UNSYNC_ASSERT(NewBlock.Offset == OldBlock.Offset);
-							UNSYNC_ASSERT(NewBlock.Size == OldBlock.Size);
-							UNSYNC_ASSERT(NewBlock.HashWeak == OldBlock.HashWeak);
-							UNSYNC_ASSERT(NewBlock.HashStrong == OldBlock.HashStrong);
-						}
-					}
-					else
-					{
-						NewFileManifest.BlockSize	= OldFileManifest.BlockSize;
-						NewFileManifest.Blocks		= std::move(OldFileManifest.Blocks);
-						NewFileManifest.MacroBlocks = std::move(OldFileManifest.MacroBlocks);
-					}
-				}
-			}
-		}
-
-		if (Params.bNeedBlocks && Params.BlockSize)
-		{
-			UNSYNC_VERBOSE(L"Updating manifest blocks");
-			UpdateDirectoryManifestBlocks(NewDirectoryManifest, Root, Params);
-		}
-	}
-	else
-	{
-		UNSYNC_VERBOSE(L"Creating manifest for '%ls'", Root.wstring().c_str());
-		NewDirectoryManifest = CreateDirectoryManifest(Root, Params);
-	}
-
-	std::swap(Result, NewDirectoryManifest);
-
-	return true;
-}
-
-template<typename T>
-inline auto
-UpdateHashT(blake3_hasher& Hasher, const T& V)
-{
-	blake3_hasher_update(&Hasher, &V, sizeof(T));
-}
-
-void
-UpdateHashBlocks(blake3_hasher& Hasher, const FGenericBlockArray& Blocks)
-{
-	for (const FGenericBlock& Block : Blocks)
-	{
-		UpdateHashT(Hasher, Block.Offset);
-		UpdateHashT(Hasher, Block.Size);
-		UpdateHashT(Hasher, Block.HashWeak);
-		blake3_hasher_update(&Hasher, Block.HashStrong.Data, Block.HashStrong.Size());
-	}
-}
-
-FHash256
-ComputeManifestStableSignature(const FDirectoryManifest& Manifest)
-{
-	blake3_hasher Hasher;
-	blake3_hasher_init(&Hasher);
-
-	UpdateHashT(Hasher, Manifest.Algorithm.ChunkingAlgorithmId);
-	UpdateHashT(Hasher, Manifest.Algorithm.WeakHashAlgorithmId);
-	UpdateHashT(Hasher, Manifest.Algorithm.StrongHashAlgorithmId);
-
-	std::vector<std::wstring_view> SortedFiles;
-	SortedFiles.reserve(Manifest.Files.size());
-
-	for (const auto& It : Manifest.Files)
-	{
-		const std::wstring& FileName = It.first;
-		SortedFiles.push_back(FileName);
-	}
-
-	std::sort(SortedFiles.begin(), SortedFiles.end());
-
-	std::string FileNameUtf8;
-	std::wstring FileName;
-	for (const std::wstring_view& FileNameView : SortedFiles)
-	{
-		FileName.clear();
-		FileName.append(FileNameView);
-
-		// Canonical unsync file paths are utf8 with unix-style separator `/`
-		ConvertWideToUtf8(FileName, FileNameUtf8);
-		std::replace(FileNameUtf8.begin(), FileNameUtf8.end(), '\\', '/');
-
-		const FFileManifest& FileManifest = Manifest.Files.at(FileName);  // TODO: heterogeneous map
-
-		blake3_hasher_update(&Hasher, FileNameUtf8.c_str(), FileNameUtf8.length());
-
-		UpdateHashT(Hasher, FileManifest.Mtime);
-		UpdateHashT(Hasher, FileManifest.Size);
-		UpdateHashT(Hasher, FileManifest.BlockSize);
-
-		UpdateHashBlocks(Hasher, FileManifest.Blocks);
-		UpdateHashBlocks(Hasher, FileManifest.MacroBlocks);
-	}
-
-	FHash256 Result;
-	blake3_hasher_finalize(&Hasher, Result.Data, sizeof(Result.Data));
-
-	return Result;
-}
-
-FHash160
-ComputeManifestStableSignature_160(const FDirectoryManifest& Manifest)
-{
-	return ToHash160(ComputeManifestStableSignature(Manifest));
-}
-
 bool
 ValidateTarget(FIOReader& Reader, const FNeedList& NeedList, EStrongHashAlgorithmID StrongHasher)
 {
@@ -4690,102 +4249,6 @@ ToString(EFileSyncStatus Status)
 	}
 }
 
-FDirectoryManifestInfo
-GetManifestInfo(const FDirectoryManifest& Manifest, bool bGenerateSignature)
-{
-	FDirectoryManifestInfo Result = {};
-
-	Result.NumBlocks	  = 0;
-	Result.NumMacroBlocks = 0;
-	Result.TotalSize	  = 0;
-	Result.UniqueSize	  = 0;
-
-	THashSet<FGenericHash> UniqueBlockSet;
-	THashSet<FGenericHash> UniqueMacroBlockSet;
-
-	for (const auto& It : Manifest.Files)
-	{
-		const FFileManifest& File = It.second;
-
-		for (const FGenericBlock& Block : File.Blocks)
-		{
-			if (UniqueBlockSet.insert(Block.HashStrong).second)
-			{
-				Result.UniqueSize += Block.Size;
-			}
-		}
-
-		for (const FGenericBlock& Block : File.MacroBlocks)
-		{
-			UniqueMacroBlockSet.insert(Block.HashStrong);
-		}
-
-		Result.TotalSize += File.Size;
-	}
-
-	Result.NumBlocks	  = UniqueBlockSet.size();
-	Result.NumMacroBlocks = UniqueMacroBlockSet.size();
-	Result.NumFiles		  = Manifest.Files.size();
-	Result.Algorithm	  = Manifest.Algorithm;
-
-	if (bGenerateSignature)
-	{
-		Result.StableSignature = ComputeManifestStableSignature(Manifest);
-	}
-
-	return Result;
-}
-
-void
-LogManifestInfo(ELogLevel LogLevel, const FDirectoryManifestInfo& Info)
-{
-	FHash160	ManifestSignature = ToHash160(Info.StableSignature);
-	const FHash160 EmptySignature = {};
-
-	if (ManifestSignature != EmptySignature)
-	{
-		std::string SignatureHexStr = HashToHexString(ManifestSignature);
-		LogPrintf(LogLevel, L"Manifest signature: %hs\n", SignatureHexStr.c_str());
-	}
-
-	LogPrintf(LogLevel, L"Chunking mode: %hs\n", ToString(Info.Algorithm.ChunkingAlgorithmId));
-	LogPrintf(LogLevel, L"Weak hash: %hs\n", ToString(Info.Algorithm.WeakHashAlgorithmId));
-	LogPrintf(LogLevel, L"Strong hash: %hs\n", ToString(Info.Algorithm.StrongHashAlgorithmId));
-	LogPrintf(LogLevel, L"Files: %llu\n", llu(Info.NumFiles));
-	LogPrintf(LogLevel, L"Blocks: %llu\n", llu(Info.NumBlocks));
-	LogPrintf(LogLevel, L"Macro blocks: %llu\n", llu(Info.NumMacroBlocks));
-	LogPrintf(LogLevel, L"Unique data size: %.0f MB (%llu bytes)\n", SizeMb(Info.UniqueSize), llu(Info.UniqueSize));
-	LogPrintf(LogLevel, L"Total data size: %.0f MB (%llu bytes)\n", SizeMb(Info.TotalSize), llu(Info.TotalSize));
-
-	// TODO: block size distribution histogram
-	// TODO: size distribution per file extension
-}
-
-void
-LogManifestInfo(ELogLevel LogLevel, const FDirectoryManifest& Manifest)
-{
-	FDirectoryManifestInfo Info = GetManifestInfo(Manifest);
-	LogManifestInfo(LogLevel, Info);
-}
-
-void
-LogManifestFiles(ELogLevel LogLevel, const FDirectoryManifest& Manifest)
-{
-	std::vector<std::wstring> Files;
-	for (const auto& ManifestIt : Manifest.Files)
-	{
-		Files.push_back(ManifestIt.first);
-	}
-
-	std::sort(Files.begin(), Files.end());
-
-	for (const std::wstring& Filename : Files)
-	{
-		const FFileManifest& Info = Manifest.Files.at(Filename);
-		LogPrintf(LogLevel, L"%ls : %llu\n", Filename.c_str(), Info.Size);
-	}
-}
-
 THashMap<FGenericHash, FGenericBlock>
 BuildBlockMap(const FDirectoryManifest& Manifest, bool bNeedMacroBlocks)
 {
@@ -4976,7 +4439,7 @@ ComputeSerializedManifestHash(const FDirectoryManifest& Manifest)
 }
 
 FHash160
-ComputeSerializedManifestHash_160(const FDirectoryManifest& Manifest)
+ComputeSerializedManifestHash160(const FDirectoryManifest& Manifest)
 {
 	return ToHash160(ComputeSerializedManifestHash(Manifest));
 }
diff --git a/Engine/Source/Programs/Unsync/Private/UnsyncCore.h b/Engine/Source/Programs/Unsync/Private/UnsyncCore.h
--- a/Engine/Source/Programs/Unsync/Private/UnsyncCore.h
+++ b/Engine/Source/Programs/Unsync/Private/UnsyncCore.h
@@ -8,6 +8,7 @@
 #include "UnsyncProtocol.h"
 #include "UnsyncSocket.h"
 #include "UnsyncUtil.h"
+#include "UnsyncManifest.h"
 
 #include <stdint.h>
 #include <unordered_map>
@@ -106,87 +107,6 @@ struct FNeedListSize
 
 FNeedListSize ComputeNeedListSize(const FNeedList& NeedList);
 
-using FGenericBlockArray = std::vector<FGenericBlock>;
-
-struct FFileManifest
-{
-	uint64 Mtime	 = 0;
-	uint64 Size		 = 0;
-	uint32 BlockSize = 0;
-
-	// runtime-only data:
-	FGenericBlockArray Blocks;
-	FGenericBlockArray MacroBlocks;
-	FPath			   CurrentPath;
-	bool			   bReadOnly = false;
-
-	bool IsValid() const { return Mtime != 0 && Size != 0; }
-};
-
-struct FAlgorithmOptionsV5
-{
-	static constexpr uint64 MAGIC	= 0x96FF3B56283EC7DBull;
-	static constexpr uint64 VERSION = 1;
-
-	EChunkingAlgorithmID   ChunkingAlgorithmId	 = EChunkingAlgorithmID::VariableBlocks;
-	EWeakHashAlgorithmID   WeakHashAlgorithmId	 = EWeakHashAlgorithmID::BuzHash;
-	EStrongHashAlgorithmID StrongHashAlgorithmId = EStrongHashAlgorithmID::Blake3_128;
-};
-
-struct FAlgorithmOptions : FAlgorithmOptionsV5
-{
-};
-
-struct FDirectoryManifest
-{
-	enum EVersions : uint64
-	{
-		Invalid = 0,
-		V1_Unsupported,
-		V2_Unsupported,
-		V3_Unsupported,
-		V4,					  // unsync v1.0.7
-		V5,					  // unsync v1.0.8: added options
-		V6_VariableHash,	  // unsync v1.0.32-iohash: added variable size hashes, moved options section before files
-		V7_OptionalSections,  // unsync v1.0.32-iohash
-
-		Latest = V7_OptionalSections
-	};
-
-	static constexpr uint64 MAGIC	= 0x80F4CC2A414F4E41ull;
-	static constexpr uint64 VERSION = EVersions::Latest;
-
-	// wide string at runtime, utf8 serialized
-	// TODO: keep paths in canonical form (utf-8, unix-style separators)
-	using FFileMap = std::unordered_map<std::wstring, FFileManifest>;
-	FFileMap Files;
-
-	// runtime data
-	FAlgorithmOptions Algorithm = {};
-	uint64			  Version = 0;
-
-	bool IsValid() const
-	{
-		return Version != EVersions::Invalid;
-	}
-};
-
-struct FDirectoryManifestInfo
-{
-	uint64	 TotalSize		 = 0;
-	uint64	 UniqueSize		 = 0;
-	uint64	 NumBlocks		 = 0;
-	uint64	 NumMacroBlocks	 = 0;
-	uint64	 NumFiles		 = 0;
-	FHash256 StableSignature = {};
-
-	FAlgorithmOptions Algorithm = {};
-};
-
-FDirectoryManifestInfo GetManifestInfo(const FDirectoryManifest& Manifest, bool bGenerateSignature = true);
-void				   LogManifestInfo(ELogLevel LogLevel, const FDirectoryManifestInfo& Info);
-void				   LogManifestFiles(ELogLevel LogLevel, const FDirectoryManifestInfo& Info);
-
 const std::string& GetVersionString();
 
 using FOnBlockGenerated = std::function<void(const FGenericBlock& Block, FBufferView Data)>;
@@ -421,21 +341,6 @@ struct FSyncDirectoryOptions
 
 bool SyncDirectory(const FSyncDirectoryOptions& SyncOptions);
 
-void			   UpdateDirectoryManifestBlocks(FDirectoryManifest& Result, const FPath& Root, const FComputeBlocksParams& Params);
-FDirectoryManifest CreateDirectoryManifest(const FPath& Root, const FComputeBlocksParams& Params);
-FDirectoryManifest CreateDirectoryManifestIncremental(const FPath& Root, const FComputeBlocksParams& Params);
-bool			   LoadOrCreateDirectoryManifest(FDirectoryManifest& Result, const FPath& Root, const FComputeBlocksParams& Params);
-
-// Computes a Blake3 hash of the manifest blocks and file metadata, ignoring any other metadata.
-// Files are processed in sorted order, with file names treated as utf-8.
-// This produces a relatively stable manifest key that does not depend on the serialization differences between versions.
-FHash256 ComputeManifestStableSignature(const FDirectoryManifest& Manifest);
-FHash160 ComputeManifestStableSignature160(const FDirectoryManifest& Manifest);
-
-// Computes a Blake3 hash of the serialized manifest data
-FHash256 ComputeSerializedManifestHash(const FDirectoryManifest& Manifest);
-FHash160 ComputeSerializedManifestHash160(const FDirectoryManifest& Manifest);
-
 // #wip-widehash -- temporary helper functions
 FBlock128			   ToBlock128(const FGenericBlock& GenericBlock);
 std::vector<FBlock128> ToBlock128(FGenericBlockArray& GenericBlocks);
diff --git a/Engine/Source/Programs/Unsync/Private/UnsyncManifest.cpp b/Engine/Source/Programs/Unsync/Private/UnsyncManifest.cpp
new file mode 100644
--- /dev/null
+++ b/Engine/Source/Programs/Unsync/Private/UnsyncManifest.cpp
@@ -0,0 +1,555 @@
+// Copyright Epic Games, Inc. All Rights Reserved.
+
+#include "UnsyncManifest.h"
+#include "UnsyncCore.h"
+#include "UnsyncFile.h"
+#include "UnsyncHash.h"
+#include "UnsyncHashTable.h"
+#include "UnsyncSerialization.h"
+#include "UnsyncThread.h"
+
+UNSYNC_THIRD_PARTY_INCLUDES_START
+#include <blake3.h>
+UNSYNC_THIRD_PARTY_INCLUDES_END
+
+namespace unsync {
+
+FDirectoryManifestInfo
+GetManifestInfo(const FDirectoryManifest& Manifest, bool bGenerateSignature)
+{
+	FDirectoryManifestInfo Result = {};
+
+	Result.NumBlocks	  = 0;
+	Result.NumMacroBlocks = 0;
+	Result.TotalSize	  = 0;
+	Result.UniqueSize	  = 0;
+
+	THashSet<FGenericHash> UniqueBlockSet;
+	THashSet<FGenericHash> UniqueMacroBlockSet;
+
+	for (const auto& It : Manifest.Files)
+	{
+		const FFileManifest& File = It.second;
+
+		for (const FGenericBlock& Block : File.Blocks)
+		{
+			if (UniqueBlockSet.insert(Block.HashStrong).second)
+			{
+				Result.UniqueSize += Block.Size;
+			}
+		}
+
+		for (const FGenericBlock& Block : File.MacroBlocks)
+		{
+			UniqueMacroBlockSet.insert(Block.HashStrong);
+		}
+
+		Result.TotalSize += File.Size;
+	}
+
+	Result.NumBlocks	  = UniqueBlockSet.size();
+	Result.NumMacroBlocks = UniqueMacroBlockSet.size();
+	Result.NumFiles		  = Manifest.Files.size();
+	Result.Algorithm	  = Manifest.Algorithm;
+
+	if (bGenerateSignature)
+	{
+		Result.StableSignature = ComputeManifestStableSignature(Manifest);
+	}
+
+	return Result;
+}
+
+void
+LogManifestInfo(ELogLevel LogLevel, const FDirectoryManifestInfo& Info)
+{
+	FHash160	   ManifestSignature = ToHash160(Info.StableSignature);
+	const FHash160 EmptySignature	 = {};
+
+	if (ManifestSignature != EmptySignature)
+	{
+		std::string SignatureHexStr = HashToHexString(ManifestSignature);
+		LogPrintf(LogLevel, L"Manifest signature: %hs\n", SignatureHexStr.c_str());
+	}
+
+	LogPrintf(LogLevel, L"Chunking mode: %hs\n", ToString(Info.Algorithm.ChunkingAlgorithmId));
+	LogPrintf(LogLevel, L"Weak hash: %hs\n", ToString(Info.Algorithm.WeakHashAlgorithmId));
+	LogPrintf(LogLevel, L"Strong hash: %hs\n", ToString(Info.Algorithm.StrongHashAlgorithmId));
+	LogPrintf(LogLevel, L"Files: %llu\n", llu(Info.NumFiles));
+	LogPrintf(LogLevel, L"Blocks: %llu\n", llu(Info.NumBlocks));
+	LogPrintf(LogLevel, L"Macro blocks: %llu\n", llu(Info.NumMacroBlocks));
+	LogPrintf(LogLevel, L"Unique data size: %.0f MB (%llu bytes)\n", SizeMb(Info.UniqueSize), llu(Info.UniqueSize));
+	LogPrintf(LogLevel, L"Total data size: %.0f MB (%llu bytes)\n", SizeMb(Info.TotalSize), llu(Info.TotalSize));
+
+	// TODO: block size distribution histogram
+	// TODO: size distribution per file extension
+}
+
+void
+LogManifestInfo(ELogLevel LogLevel, const FDirectoryManifest& Manifest)
+{
+	FDirectoryManifestInfo Info = GetManifestInfo(Manifest);
+	LogManifestInfo(LogLevel, Info);
+}
+
+void
+LogManifestFiles(ELogLevel LogLevel, const FDirectoryManifest& Manifest)
+{
+	std::vector<std::wstring> Files;
+	for (const auto& ManifestIt : Manifest.Files)
+	{
+		Files.push_back(ManifestIt.first);
+	}
+
+	std::sort(Files.begin(), Files.end());
+
+	for (const std::wstring& Filename : Files)
+	{
+		const FFileManifest& Info = Manifest.Files.at(Filename);
+		LogPrintf(LogLevel, L"%ls : %llu\n", Filename.c_str(), Info.Size);
+	}
+}
+
+template<typename T>
+inline auto
+UpdateHashT(blake3_hasher& Hasher, const T& V)
+{
+	blake3_hasher_update(&Hasher, &V, sizeof(T));
+}
+
+void
+UpdateHashBlocks(blake3_hasher& Hasher, const FGenericBlockArray& Blocks)
+{
+	for (const FGenericBlock& Block : Blocks)
+	{
+		UpdateHashT(Hasher, Block.Offset);
+		UpdateHashT(Hasher, Block.Size);
+		UpdateHashT(Hasher, Block.HashWeak);
+		blake3_hasher_update(&Hasher, Block.HashStrong.Data, Block.HashStrong.Size());
+	}
+}
+
+FHash256
+ComputeManifestStableSignature(const FDirectoryManifest& Manifest)
+{
+	blake3_hasher Hasher;
+	blake3_hasher_init(&Hasher);
+
+	UpdateHashT(Hasher, Manifest.Algorithm.ChunkingAlgorithmId);
+	UpdateHashT(Hasher, Manifest.Algorithm.WeakHashAlgorithmId);
+	UpdateHashT(Hasher, Manifest.Algorithm.StrongHashAlgorithmId);
+
+	std::vector<std::wstring_view> SortedFiles;
+	SortedFiles.reserve(Manifest.Files.size());
+
+	for (const auto& It : Manifest.Files)
+	{
+		const std::wstring& FileName = It.first;
+		SortedFiles.push_back(FileName);
+	}
+
+	std::sort(SortedFiles.begin(), SortedFiles.end());
+
+	std::string	 FileNameUtf8;
+	std::wstring FileName;
+	for (const std::wstring_view& FileNameView : SortedFiles)
+	{
+		FileName.clear();
+		FileName.append(FileNameView);
+
+		// Canonical unsync file paths are utf8 with unix-style separator `/`
+		ConvertWideToUtf8(FileName, FileNameUtf8);
+		std::replace(FileNameUtf8.begin(), FileNameUtf8.end(), '\\', '/');
+
+		const FFileManifest& FileManifest = Manifest.Files.at(FileName);  // TODO: heterogeneous map
+
+		blake3_hasher_update(&Hasher, FileNameUtf8.c_str(), FileNameUtf8.length());
+
+		UpdateHashT(Hasher, FileManifest.Mtime);
+		UpdateHashT(Hasher, FileManifest.Size);
+		UpdateHashT(Hasher, FileManifest.BlockSize);
+
+		UpdateHashBlocks(Hasher, FileManifest.Blocks);
+		UpdateHashBlocks(Hasher, FileManifest.MacroBlocks);
+	}
+
+	FHash256 Result;
+	blake3_hasher_finalize(&Hasher, Result.Data, sizeof(Result.Data));
+
+	return Result;
+}
+
+FHash160
+ComputeManifestStableSignature160(const FDirectoryManifest& Manifest)
+{
+	return ToHash160(ComputeManifestStableSignature(Manifest));
+}
+
+void
+UpdateDirectoryManifestBlocks(FDirectoryManifest& Result, const FPath& Root, const FComputeBlocksParams& Params)
+{
+	UNSYNC_LOG_INDENT;
+
+	UNSYNC_ASSERT(Params.bNeedBlocks);
+	UNSYNC_ASSERT(Params.BlockSize != 0);
+
+	FTimePoint TimeBegin = TimePointNow();
+
+	uint32 NumProcessedFiles = 0;
+
+	FTaskGroup TaskGroup;
+
+	const uint32 MaxConcurrentFiles = 8;  // quickly diminishing returns past 8 concurrent files
+	FSemaphore	 Semaphore(MaxConcurrentFiles);
+
+	uint64 NumSkippedBlocks = 0;
+	uint64 NumSkippedBytes	= 0;
+
+	for (auto& It : Result.Files)
+	{
+		FFileManifest& FileManifest	  = It.second;
+		const bool	   bBlockSizeOk	  = FileManifest.BlockSize == Params.BlockSize;
+		const bool	   bMacroBlocksOk = Params.bNeedMacroBlocks ? !FileManifest.MacroBlocks.empty() : true;
+		if (bBlockSizeOk && bMacroBlocksOk)
+		{
+			NumSkippedBlocks += FileManifest.Blocks.size();
+			NumSkippedBytes += FileManifest.Size;
+
+			continue;
+		}
+
+		++NumProcessedFiles;
+
+		FPath FilePath = Root / It.first;
+
+		Semaphore.Acquire();
+
+		UNSYNC_VERBOSE(L"Computing blocks for '%ls' (%.2f MB)", FilePath.wstring().c_str(), SizeMb(It.second.Size));
+		auto BlockTask = [&FileManifest, &Semaphore, &Params, FilePath = std::move(FilePath)]()
+		{
+			FNativeFile File(FilePath, EFileMode::ReadOnlyUnbuffered);
+			if (File.IsValid())
+			{
+				FComputeBlocksResult ComputedBlocks = ComputeBlocks(File, Params);
+				std::swap(FileManifest.Blocks, ComputedBlocks.Blocks);
+				std::swap(FileManifest.MacroBlocks, ComputedBlocks.MacroBlocks);
+
+				FileManifest.BlockSize = Params.BlockSize;
+			}
+			else
+			{
+				UNSYNC_FATAL(L"Failed to open file '%ls' while computing manifest blocks. %hs",
+							 FilePath.wstring().c_str(),
+							 FormatSystemErrorMessage(File.GetError()).c_str());
+			}
+			Semaphore.Release();
+		};
+
+		TaskGroup.run(BlockTask);
+	}
+
+	TaskGroup.wait();
+
+	// TODO: move manifest stat printing into a helper
+
+	uint64 NumTotalMacroBlocks = 0;
+	uint64 NumTotalBlocks	   = 0;
+	uint64 NumTotalBytes	   = 0;
+
+	for (auto& It : Result.Files)
+	{
+		FFileManifest& FileManifest = It.second;
+		NumTotalMacroBlocks += FileManifest.MacroBlocks.size();
+		NumTotalBlocks += FileManifest.Blocks.size();
+		NumTotalBytes += FileManifest.Size;
+	}
+
+	NumTotalBlocks -= NumSkippedBlocks;
+	NumTotalBytes -= NumSkippedBytes;
+
+	double Duration = DurationSec(TimeBegin, TimePointNow());
+
+	if (NumTotalBlocks == 0)
+	{
+		UNSYNC_VERBOSE(L"No blocks needed to be computed")
+	}
+	else
+	{
+		UNSYNC_VERBOSE(L"Total computed %lld block(s) in %.3f sec (%.3f MB, %.3f MB / sec)",
+					   (long long)NumTotalBlocks,
+					   Duration,
+					   SizeMb(NumTotalBytes),
+					   SizeMb((double(NumTotalBytes) / Duration)));
+
+		if (NumTotalMacroBlocks)
+		{
+			UNSYNC_VERBOSE(L"Macro blocks: %lld", NumTotalMacroBlocks);
+		}
+	}
+}
+
+FDirectoryManifest
+CreateDirectoryManifest(const FPath& Root, const FComputeBlocksParams& Params)
+{
+	UNSYNC_LOG_INDENT;
+
+	FDirectoryManifest Result;
+
+	Result.Algorithm = Params.Algorithm;
+
+	FTimePoint TimeBegin = TimePointNow();
+
+	FTaskGroup	 TaskGroup;
+	const uint32 MaxConcurrentFiles = 8;  // quickly diminishing returns past 8 concurrent files
+	FSemaphore	 Semaphore(MaxConcurrentFiles);
+
+	std::mutex ResultMutex;
+	FPath	   UnsyncDirName = ".unsync";
+
+	for (const std::filesystem::directory_entry& Dir : RecursiveDirectoryScan(Root))
+	{
+		if (Dir.is_directory())
+		{
+			continue;
+		}
+
+		FPath RelativePath = GetRelativePath(Dir.path(), Root);
+
+		if (RelativePath.native().starts_with(UnsyncDirName.native()))
+		{
+			continue;
+		}
+
+		const std::wstring PathKey = RelativePath.wstring();
+
+		UNSYNC_VERBOSE2(L"Found '%ls'", PathKey.c_str());
+
+		FFileManifest FileManifest;
+
+		FileManifest.Mtime		 = ToWindowsFileTime(Dir.last_write_time());
+		FileManifest.Size		 = Dir.file_size();
+		FileManifest.CurrentPath = Dir.path();
+		FileManifest.BlockSize	 = Params.BlockSize;
+		FileManifest.bReadOnly	 = IsReadOnly(Dir.status().permissions());
+
+		{
+			std::lock_guard<std::mutex> LockGuard(ResultMutex);
+			Result.Files[PathKey] = std::move(FileManifest);
+		}
+
+		if (Params.bNeedBlocks && Params.BlockSize)
+		{
+			FPath FilePath = Root / RelativePath;
+			auto  File	   = std::make_shared<FNativeFile>(FilePath, EFileMode::ReadOnlyUnbuffered);
+			if (File->IsValid())
+			{
+				UNSYNC_VERBOSE(L"Computing blocks for '%ls' (%.2f MB)", FilePath.wstring().c_str(), double(File->GetSize()) / (1 << 20));
+
+				Semaphore.Acquire();
+				TaskGroup.run(
+					[&Semaphore, &ResultMutex, &Result, File = std::move(File), Key = std::move(PathKey), &Params]()
+					{
+						FComputeBlocksResult ComputedBlocks = ComputeBlocks(*File, Params);
+
+						std::lock_guard<std::mutex> LockGuard(ResultMutex);
+
+						std::swap(Result.Files[Key].Blocks, ComputedBlocks.Blocks);
+						std::swap(Result.Files[Key].MacroBlocks, ComputedBlocks.MacroBlocks);
+
+						Semaphore.Release();
+					});
+			}
+			else
+			{
+				UNSYNC_FATAL(L"Failed to open file '%ls' while computing manifest blocks. %hs",
+							 FilePath.wstring().c_str(),
+							 FormatSystemErrorMessage(File->GetError()).c_str());
+			}
+		}
+	}
+
+	TaskGroup.wait();
+
+	if (Params.bNeedBlocks && Params.BlockSize)
+	{
+		// TODO: move manifest stat printing into a helper
+
+		uint64 NumTotalMacroBlocks = 0;
+		uint64 NumTotalBlocks	   = 0;
+		uint64 NumTotalBytes	   = 0;
+
+		for (auto& It : Result.Files)
+		{
+			FFileManifest& FileManifest = It.second;
+			NumTotalMacroBlocks += FileManifest.MacroBlocks.size();
+			NumTotalBlocks += FileManifest.Blocks.size();
+			NumTotalBytes += FileManifest.Size;
+		}
+
+		double Duration = DurationSec(TimeBegin, TimePointNow());
+		UNSYNC_VERBOSE(L"Total computed %lld block(s) in %.3f sec (%.3f MB, %.3f MB / sec)",
+					   (long long)NumTotalBlocks,
+					   Duration,
+					   SizeMb(NumTotalBytes),
+					   SizeMb((double(NumTotalBytes) / Duration)));
+		if (NumTotalMacroBlocks)
+		{
+			UNSYNC_VERBOSE(L"Macro blocks: %lld", NumTotalMacroBlocks);
+		}
+	}
+
+	return Result;
+}
+
+FDirectoryManifest
+CreateDirectoryManifestIncremental(const FPath& Root, const FComputeBlocksParams& InParams)
+{
+	FComputeBlocksParams Params = InParams;
+
+	FPath ManifestRoot			= Root / ".unsync";
+	FPath DirectoryManifestPath = ManifestRoot / "manifest.bin";
+
+	FDirectoryManifest OldManifest;
+	const bool		   bExistingManifestLoaded = LoadDirectoryManifest(OldManifest, Root, DirectoryManifestPath);
+
+	// Inherit algorithm options from the existing manifest
+	if (bExistingManifestLoaded)
+	{
+		Params.Algorithm = OldManifest.Algorithm;
+	}
+
+	// Scan the input directory and gather file metadata, without generating blocks
+	FComputeBlocksParams LightweightManifestParams = Params;
+	LightweightManifestParams.bNeedBlocks		   = false;
+	LightweightManifestParams.BlockSize			   = 0;
+
+	FDirectoryManifest NewManifest = CreateDirectoryManifest(Root, LightweightManifestParams);
+
+	// Copy file blocks from old manifest, if possible
+	for (auto& NewManifestFileEntry : NewManifest.Files)
+	{
+		const std::wstring& FileName			 = NewManifestFileEntry.first;
+		auto				OldManifestFileEntry = OldManifest.Files.find(FileName);
+		if (OldManifestFileEntry == OldManifest.Files.end())
+		{
+			continue;
+		}
+
+		FFileManifest& NewEntry = NewManifestFileEntry.second;
+		FFileManifest& OldEntry = OldManifestFileEntry->second;
+
+		if (NewEntry.Mtime == OldEntry.Mtime && NewEntry.Size == OldEntry.Size)
+		{
+			NewEntry.Blocks		 = std::move(OldEntry.Blocks);
+			NewEntry.MacroBlocks = std::move(OldEntry.MacroBlocks);
+			NewEntry.BlockSize	 = OldEntry.BlockSize;
+		}
+	}
+
+	// Generate blocks for changed or new files
+	UpdateDirectoryManifestBlocks(NewManifest, Root, Params);
+
+	return NewManifest;
+}
+
+bool
+AlgorithmOptionsCompatible(const FAlgorithmOptions& A, const FAlgorithmOptions& B)
+{
+	return A.StrongHashAlgorithmId == B.StrongHashAlgorithmId && B.WeakHashAlgorithmId == B.WeakHashAlgorithmId &&
+		   B.ChunkingAlgorithmId == B.ChunkingAlgorithmId;
+}
+
+bool
+LoadOrCreateDirectoryManifest(FDirectoryManifest& Result, const FPath& Root, const FComputeBlocksParams& Params)
+{
+	UNSYNC_LOG_INDENT;
+
+	FPath ManifestRoot			= Root / ".unsync";
+	FPath DirectoryManifestPath = ManifestRoot / "manifest.bin";
+
+	FDirectoryManifest OldDirectoryManifest;
+	FDirectoryManifest NewDirectoryManifest;
+
+	const bool bManifestFileExists = PathExists(DirectoryManifestPath);
+	if (!bManifestFileExists)
+	{
+		UNSYNC_VERBOSE(L"Manifest file '%ls' does not exist", DirectoryManifestPath.wstring().c_str());
+	}
+
+	const bool bExistingManifestLoaded = bManifestFileExists && LoadDirectoryManifest(OldDirectoryManifest, Root, DirectoryManifestPath);
+	const bool bExistingManifestCompatible = AlgorithmOptionsCompatible(OldDirectoryManifest.Algorithm, Params.Algorithm);
+
+	if (bExistingManifestLoaded && bExistingManifestCompatible)
+	{
+		UNSYNC_VERBOSE(L"Loaded existing manifest from '%ls'", DirectoryManifestPath.wstring().c_str());
+
+		UNSYNC_VERBOSE(L"Creating lightweight manifest for '%ls'", Root.wstring().c_str());
+		{
+			// Verify that manifests match in dry run mode.
+			// Otherwise just do a quick manifest generation, without file blocks.
+			FComputeBlocksParams NewParams = Params;
+			if (GDryRun)
+			{
+				NewParams.bNeedBlocks = true;
+			}
+			else
+			{
+				NewParams.bNeedBlocks = false;
+				NewParams.BlockSize	  = 0;
+			}
+
+			NewDirectoryManifest = CreateDirectoryManifest(Root, NewParams);
+		}
+
+		UNSYNC_VERBOSE(L"Comparing manifests");
+		for (const auto& OldManifestIt : OldDirectoryManifest.Files)
+		{
+			auto NewManifestIt = NewDirectoryManifest.Files.find(OldManifestIt.first);
+			if (NewManifestIt != NewDirectoryManifest.Files.end())
+			{
+				const FFileManifest& OldFileManifest = OldManifestIt.second;
+				FFileManifest&		 NewFileManifest = NewManifestIt->second;
+				if (NewFileManifest.Size == OldFileManifest.Size && NewFileManifest.Mtime == OldFileManifest.Mtime)
+				{
+					if (NewFileManifest.BlockSize)
+					{
+						UNSYNC_ASSERT(NewFileManifest.BlockSize == OldFileManifest.BlockSize);
+						UNSYNC_ASSERT(NewFileManifest.Blocks.size() == OldFileManifest.Blocks.size());
+						for (uint64 I = 0; I < NewFileManifest.Blocks.size(); ++I)
+						{
+							const FGenericBlock& NewBlock = NewFileManifest.Blocks[I];
+							const FGenericBlock& OldBlock = OldFileManifest.Blocks[I];
+							UNSYNC_ASSERT(NewBlock.Offset == OldBlock.Offset);
+							UNSYNC_ASSERT(NewBlock.Size == OldBlock.Size);
+							UNSYNC_ASSERT(NewBlock.HashWeak == OldBlock.HashWeak);
+							UNSYNC_ASSERT(NewBlock.HashStrong == OldBlock.HashStrong);
+						}
+					}
+					else
+					{
+						NewFileManifest.BlockSize	= OldFileManifest.BlockSize;
+						NewFileManifest.Blocks		= std::move(OldFileManifest.Blocks);
+						NewFileManifest.MacroBlocks = std::move(OldFileManifest.MacroBlocks);
+					}
+				}
+			}
+		}
+
+		if (Params.bNeedBlocks && Params.BlockSize)
+		{
+			UNSYNC_VERBOSE(L"Updating manifest blocks");
+			UpdateDirectoryManifestBlocks(NewDirectoryManifest, Root, Params);
+		}
+	}
+	else
+	{
+		UNSYNC_VERBOSE(L"Creating manifest for '%ls'", Root.wstring().c_str());
+		NewDirectoryManifest = CreateDirectoryManifest(Root, Params);
+	}
+
+	std::swap(Result, NewDirectoryManifest);
+
+	return true;
+}
+
+}  // namespace unsync
diff --git a/Engine/Source/Programs/Unsync/Private/UnsyncManifest.h b/Engine/Source/Programs/Unsync/Private/UnsyncManifest.h
new file mode 100644
--- /dev/null
+++ b/Engine/Source/Programs/Unsync/Private/UnsyncManifest.h
@@ -0,0 +1,107 @@
+// Copyright Epic Games, Inc. All Rights Reserved.
+
+#pragma once
+
+#include "UnsyncCommon.h"
+#include "UnsyncHash.h"
+#include "UnsyncProtocol.h"
+
+namespace unsync {
+
+struct FComputeBlocksParams;
+
+struct FFileManifest
+{
+	uint64 Mtime	 = 0;
+	uint64 Size		 = 0;
+	uint32 BlockSize = 0;
+
+	// runtime-only data:
+	FGenericBlockArray Blocks;
+	FGenericBlockArray MacroBlocks;
+	FPath			   CurrentPath;
+	bool			   bReadOnly = false;
+
+	bool IsValid() const { return Mtime != 0 && Size != 0; }
+};
+
+struct FAlgorithmOptionsV5
+{
+	static constexpr uint64 MAGIC	= 0x96FF3B56283EC7DBull;
+	static constexpr uint64 VERSION = 1;
+
+	EChunkingAlgorithmID   ChunkingAlgorithmId	 = EChunkingAlgorithmID::VariableBlocks;
+	EWeakHashAlgorithmID   WeakHashAlgorithmId	 = EWeakHashAlgorithmID::BuzHash;
+	EStrongHashAlgorithmID StrongHashAlgorithmId = EStrongHashAlgorithmID::Blake3_128;
+};
+
+struct FAlgorithmOptions : FAlgorithmOptionsV5
+{
+};
+
+struct FDirectoryManifest
+{
+	enum EVersions : uint64 {
+		Invalid = 0,
+		V1_Unsupported,
+		V2_Unsupported,
+		V3_Unsupported,
+		V4,					  // unsync v1.0.7
+		V5,					  // unsync v1.0.8: added options
+		V6_VariableHash,	  // unsync v1.0.32-iohash: added variable size hashes, moved options section before files
+		V7_OptionalSections,  // unsync v1.0.32-iohash
+
+		Latest = V7_OptionalSections
+	};
+
+	static constexpr uint64 MAGIC	= 0x80F4CC2A414F4E41ull;
+	static constexpr uint64 VERSION = EVersions::Latest;
+
+	// wide string at runtime, utf8 serialized
+	// TODO: keep paths in canonical form (utf-8, unix-style separators)
+	using FFileMap = std::unordered_map<std::wstring, FFileManifest>;
+	FFileMap Files;
+
+	// runtime data
+	FAlgorithmOptions Algorithm = {};
+	uint64			  Version	= 0;
+
+	bool IsValid() const { return Version != EVersions::Invalid; }
+};
+
+struct FDirectoryManifestInfo
+{
+	uint64	 TotalSize		 = 0;
+	uint64	 UniqueSize		 = 0;
+	uint64	 NumBlocks		 = 0;
+	uint64	 NumMacroBlocks	 = 0;
+	uint64	 NumFiles		 = 0;
+	FHash256 StableSignature = {};
+
+	FAlgorithmOptions Algorithm = {};
+};
+
+FDirectoryManifestInfo GetManifestInfo(const FDirectoryManifest& Manifest, bool bGenerateSignature = true);
+void				   LogManifestInfo(ELogLevel LogLevel, const FDirectoryManifest& Manifest);
+void				   LogManifestInfo(ELogLevel LogLevel, const FDirectoryManifestInfo& Info);
+void				   LogManifestFiles(ELogLevel LogLevel, const FDirectoryManifest& Manifest);
+void				   LogManifestFiles(ELogLevel LogLevel, const FDirectoryManifestInfo& Info);
+
+// Computes a Blake3 hash of the manifest blocks and file metadata, ignoring any other metadata.
+// Files are processed in sorted order, with file names treated as utf-8.
+// This produces a relatively stable manifest key that does not depend on the serialization differences between versions.
+FHash256 ComputeManifestStableSignature(const FDirectoryManifest& Manifest);
+FHash160 ComputeManifestStableSignature160(const FDirectoryManifest& Manifest);
+
+// Computes a Blake3 hash of the serialized manifest data
+FHash256 ComputeSerializedManifestHash(const FDirectoryManifest& Manifest);
+FHash160 ComputeSerializedManifestHash160(const FDirectoryManifest& Manifest);
+
+void			   UpdateDirectoryManifestBlocks(FDirectoryManifest& Result, const FPath& Root, const FComputeBlocksParams& Params);
+FDirectoryManifest CreateDirectoryManifest(const FPath& Root, const FComputeBlocksParams& Params);
+FDirectoryManifest CreateDirectoryManifestIncremental(const FPath& Root, const FComputeBlocksParams& Params);
+bool			   LoadOrCreateDirectoryManifest(FDirectoryManifest& Result, const FPath& Root, const FComputeBlocksParams& Params);
+
+bool AlgorithmOptionsCompatible(const FAlgorithmOptions& A, const FAlgorithmOptions& B);
+
+}  // namespace unsync
diff --git a/Engine/Source/Programs/Unsync/Private/UnsyncProtocol.h b/Engine/Source/Programs/Unsync/Private/UnsyncProtocol.h
--- a/Engine/Source/Programs/Unsync/Private/UnsyncProtocol.h
+++ b/Engine/Source/Programs/Unsync/Private/UnsyncProtocol.h
@@ -64,6 +64,7 @@ using FBlock160 = TBlock<FHash160>;
 using FBlock256 = TBlock<FHash256>;
 
 using FGenericBlock = TBlock<FGenericHash>;
+using FGenericBlockArray = std::vector<FGenericBlock>;
 
 static constexpr uint64 SERIALIZED_SECTION_ID_TERMINATOR		  = 0;
 static constexpr uint64 SERIALIZED_SECTION_ID_METADATA_STRING	  = 0xC6BD6CDCEEF79533ull;
