unsync - Use a struct to pass common parameters to BuildTarget, etc.

diff --git a/Engine/Source/Programs/Unsync/Private/UnsyncCmdPack.cpp b/Engine/Source/Programs/Unsync/Private/UnsyncCmdPack.cpp
--- a/Engine/Source/Programs/Unsync/Private/UnsyncCmdPack.cpp
+++ b/Engine/Source/Programs/Unsync/Private/UnsyncCmdPack.cpp
@@ -81,14 +81,6 @@ BuildP4HaveSet(const FPath& Root, std::string_view P4HaveDataUtf8, FDirectoryMan
 	ForLines(P4HaveDataUtf8, Callback);
 }
 
-static void DropBlocksFromManifest(FDirectoryManifest& DirectoryManifest)
-{
-	for (auto& FileEntry : DirectoryManifest.Files)
-	{
-		FileEntry.second.Blocks.clear();
-	}
-}
-
 struct FPackIndexEntry
 {
 	FHash160 Hash	 = {};
@@ -284,7 +276,7 @@ int32 CmdPack(const FCmdPackOptions& Options)
 	{
 		UNSYNC_LOG(L"Loading previous manifest ... ")
 		FDirectoryManifest OldManifest;
-		if (LoadDirectoryManifest(OldManifest, InputRoot, DirectoryManifestPath))
+		if (PathExists(DirectoryManifestPath) && LoadDirectoryManifest(OldManifest, InputRoot, DirectoryManifestPath))
 		{
 			UNSYNC_LOG(L"Previous manifest loaded")
 
@@ -323,7 +315,6 @@ int32 CmdPack(const FCmdPackOptions& Options)
 	if (!GDryRun)
 	{
 		UNSYNC_LOG(L"Saving directory manifest '%ls'", DirectoryManifestPath.wstring().c_str());
-		DropBlocksFromManifest(DirectoryManifest); // only need to keep macro blocks
 		SaveDirectoryManifest(DirectoryManifest, DirectoryManifestPath);
 	}
 
@@ -365,6 +356,8 @@ int32 CmdPack(const FCmdPackOptions& Options)
 		FPath FinalPackFilename	 = ManifestRoot / (OutputId + ".bin");
 		FPath FinalIndexFilename = ManifestRoot / (OutputId + ".idx");
 
+		UNSYNC_LOG(L"Saving new pack: %hs", OutputId.c_str());
+
 		if (!FileRename(TempOutputPackFilename, FinalPackFilename, ErrorCode))
 		{
 			UNSYNC_ERROR(L"Failed to rename temporary pack file to '%ls'", FinalPackFilename.wstring().c_str());
diff --git a/Engine/Source/Programs/Unsync/Private/UnsyncCore.cpp b/Engine/Source/Programs/Unsync/Private/UnsyncCore.cpp
--- a/Engine/Source/Programs/Unsync/Private/UnsyncCore.cpp
+++ b/Engine/Source/Programs/Unsync/Private/UnsyncCore.cpp
@@ -1058,16 +1058,16 @@ OptimizeNeedList(const std::vector<FNeedBlock>& Input, uint64 MaxMergedBlockSize
 }
 
 FBuffer
-BuildTargetBuffer(const uint8*			 SourceData,
-				  uint64				 SourceSize,
-				  const uint8*			 BaseData,
-				  uint64				 BaseSize,
-				  const FNeedList&		 NeedList,
-				  EStrongHashAlgorithmID StrongHasher)
+BuildTargetBuffer(const uint8*				SourceData,
+				  uint64					SourceSize,
+				  const uint8*				BaseData,
+				  uint64					BaseSize,
+				  const FNeedList&			NeedList,
+				  const FBuildTargetParams& Params)
 {
 	FMemReader SourceReader(SourceData, SourceSize);
 	FMemReader BaseReader(BaseData, BaseSize);
-	return BuildTargetBuffer(SourceReader, BaseReader, NeedList, StrongHasher);
+	return BuildTargetBuffer(SourceReader, BaseReader, NeedList, Params);
 }
 
 struct FReadSchedule
@@ -1108,13 +1108,13 @@ public:
 };
 
 FBuffer
-BuildTargetBuffer(FIOReader& SourceProvider, FIOReader& BaseProvider, const FNeedList& NeedList, EStrongHashAlgorithmID StrongHasher)
+BuildTargetBuffer(FIOReader& SourceProvider, FIOReader& BaseProvider, const FNeedList& NeedList, const FBuildTargetParams& Params)
 {
 	FBuffer				Result;
 	const FNeedListSize SizeInfo = ComputeNeedListSize(NeedList);
 	Result.Resize(SizeInfo.TotalBytes);
 	FMemReaderWriter ResultWriter(Result.Data(), Result.Size());
-	BuildTarget(ResultWriter, SourceProvider, BaseProvider, NeedList, StrongHasher);
+	BuildTarget(ResultWriter, SourceProvider, BaseProvider, NeedList, Params);
 	return Result;
 }
 
@@ -1263,14 +1263,7 @@ DownloadBlocks(FProxyPool&					  ProxyPool,
 }
 
 FBuildTargetResult
-BuildTarget(FIOWriter&			   Output,
-			FIOReader&			   Source,
-			FIOReader&			   Base,
-			const FNeedList&	   NeedList,
-			EStrongHashAlgorithmID StrongHasher,
-			FProxyPool*			   ProxyPool,
-			FBlockCache*		   BlockCache,
-			FScavengeDatabase*	   ScavengeDatabase)
+BuildTarget(FIOWriter& Output, FIOReader& Source, FIOReader& Base, const FNeedList& NeedList, const FBuildTargetParams& Params)
 {
 	UNSYNC_LOG_INDENT;
 
@@ -1278,7 +1271,11 @@ BuildTarget(FIOWriter&			   Output,
 
 	auto TimeBegin = TimePointNow();
 
-	const FNeedListSize SizeInfo = ComputeNeedListSize(NeedList);
+	const FNeedListSize			 SizeInfo		  = ComputeNeedListSize(NeedList);
+	const EStrongHashAlgorithmID StrongHasher	  = Params.StrongHasher;
+	FProxyPool*					 ProxyPool		  = Params.ProxyPool;
+	FScavengeDatabase*			 ScavengeDatabase = Params.ScavengeDatabase;
+	FBlockCache*				 BlockCache		  = Params.BlockCache;
 
 	if (SizeInfo.TotalBytes != Output.GetSize())
 	{
@@ -1339,7 +1336,7 @@ BuildTarget(FIOWriter&			   Output,
 	}
 
 	auto ProcessNeedList =
-		[bAllowVerboseLog, LogIndent, &Output, &Error, &WriteSemaphore, &WriteTasks, &bWaitingForBaseData, &Stats, SizeInfo](
+		[bAllowVerboseLog, LogIndent, &Output, &Error, &WriteSemaphore, &WriteTasks, &bWaitingForBaseData, &Stats, &Params, SizeInfo](
 							   FIOReader&					  DataProvider,
 							   const std::vector<FNeedBlock>& NeedBlocks,
 							   uint64						  TotalCopySize,
@@ -1357,7 +1354,8 @@ BuildTarget(FIOWriter&			   Output,
 			return;
 		}
 
-		if (ListType == EBlockListType::Source && DataProvider.GetSize() != SizeInfo.TotalBytes)
+		if (ListType == EBlockListType::Source && Params.SourceType == FBuildTargetParams::ESourceType::File &&
+			DataProvider.GetSize() != SizeInfo.TotalBytes)
 		{
 			UNSYNC_ERROR(L"File size is %llu, but expected to be %llu. File may have changed after manifest was generated.",
 						 llu(DataProvider.GetSize()),
@@ -2068,14 +2066,13 @@ SyncFile(const FNeedList&		   NeedList,
 			return std::unique_ptr<FNativeFile>(new FNativeFile(SourceFilePath, EFileMode::ReadOnlyUnbuffered));
 		});
 
-		FBuildTargetResult BuildResult = BuildTarget(*TargetFile,
-													 SourceFile,
-													 BaseDataReader,
-													 NeedList,
-													 Options.Algorithm.StrongHashAlgorithmId,
-													 Options.ProxyPool,
-													 Options.BlockCache,
-													 Options.ScavengeDatabase);
+		FBuildTargetParams BuildParams;
+		BuildParams.StrongHasher	 = Options.Algorithm.StrongHashAlgorithmId;
+		BuildParams.ProxyPool		 = Options.ProxyPool;
+		BuildParams.BlockCache		 = Options.BlockCache;
+		BuildParams.ScavengeDatabase = Options.ScavengeDatabase;
+
+		FBuildTargetResult BuildResult = BuildTarget(*TargetFile, SourceFile, BaseDataReader, NeedList, BuildParams);
 
 		Result.SourceBytes = BuildResult.SourceBytes;
 		Result.BaseBytes   = BuildResult.BaseBytes;
@@ -4080,7 +4077,11 @@ BuildTargetWithPatch(const uint8* PatchData, uint64 PatchSize, const uint8* Base
 		NeedList.Sequence.push_back(Hash);
 	}
 
-	Result = BuildTargetBuffer(SourceData, SourceDataSize, BaseData, BaseDataSize, NeedList, Header.StrongHashAlgorithmId);
+	FBuildTargetParams BuildParams;
+	BuildParams.StrongHasher = Header.StrongHashAlgorithmId;
+	BuildParams.SourceType	 = FBuildTargetParams::ESourceType::Patch;
+
+	Result = BuildTargetBuffer(SourceData, SourceDataSize, BaseData, BaseDataSize, NeedList, BuildParams);
 
 	FGenericBlockArray SourceValidation;
 
diff --git a/Engine/Source/Programs/Unsync/Private/UnsyncCore.h b/Engine/Source/Programs/Unsync/Private/UnsyncCore.h
--- a/Engine/Source/Programs/Unsync/Private/UnsyncCore.h
+++ b/Engine/Source/Programs/Unsync/Private/UnsyncCore.h
@@ -189,26 +189,35 @@ struct FBuildTargetResult
 	uint64 BaseBytes   = 0;
 };
 
-FBuildTargetResult BuildTarget(FIOWriter&			  Result,
-				 FIOReader&				Source,
-				 FIOReader&				Base,
-				 const FNeedList&		NeedList,
-				 EStrongHashAlgorithmID StrongHasher,
-				 FProxyPool*			ProxyPool = nullptr,
-				 FBlockCache*			BlockCache = nullptr,
-				 FScavengeDatabase*		ScavengeDatabase = nullptr);
-
-FBuffer BuildTargetBuffer(FIOReader&			 SourceProvider,
-						  FIOReader&			 BaseProvider,
-						  const FNeedList&		 NeedList,
-						  EStrongHashAlgorithmID StrongHasher);
-
-FBuffer BuildTargetBuffer(const uint8*			 SourceData,
-						  uint64				 SourceSize,
-						  const uint8*			 BaseData,
-						  uint64				 BaseSize,
-						  const FNeedList&		 NeedList,
-						  EStrongHashAlgorithmID StrongHasher);
+struct FBuildTargetParams
+{
+	EStrongHashAlgorithmID StrongHasher;
+	FProxyPool*			   ProxyPool		= nullptr;
+	FBlockCache*		   BlockCache		= nullptr;
+	FScavengeDatabase*	   ScavengeDatabase = nullptr;
+
+	enum class ESourceType {
+		File,
+		Patch
+	};
+
+	ESourceType SourceType = ESourceType::File;
+};
+
+FBuildTargetResult BuildTarget(FIOWriter&				 Result,
+							   FIOReader&				 Source,
+							   FIOReader&				 Base,
+							   const FNeedList&			 NeedList,
+							   const FBuildTargetParams& Params);
+
+FBuffer BuildTargetBuffer(FIOReader& SourceProvider, FIOReader& BaseProvider, const FNeedList& NeedList, const FBuildTargetParams& Params);
+
+FBuffer BuildTargetBuffer(const uint8*				SourceData,
+						  uint64					SourceSize,
+						  const uint8*				BaseData,
+						  uint64					BaseSize,
+						  const FNeedList&			NeedList,
+						  const FBuildTargetParams& Params);
 
 FBuffer BuildTargetWithPatch(const uint8* PatchData, uint64 PatchSize, const uint8* BaseData, uint64 BaseSize);
 
diff --git a/Engine/Source/Programs/Unsync/Private/UnsyncTest.cpp b/Engine/Source/Programs/Unsync/Private/UnsyncTest.cpp
--- a/Engine/Source/Programs/Unsync/Private/UnsyncTest.cpp
+++ b/Engine/Source/Programs/Unsync/Private/UnsyncTest.cpp
@@ -271,6 +271,9 @@ TestSync(EWeakHashAlgorithmID WeakHasher, EStrongHashAlgorithmID StrongHasher)
 	Algorithm.ChunkingAlgorithmId	= ChunkingMode;
 	Algorithm.WeakHashAlgorithmId	= WeakHasher;
 	Algorithm.StrongHashAlgorithmId = StrongHasher;
+	
+	FBuildTargetParams BuildParams;
+	BuildParams.StrongHasher = StrongHasher;
 
 	UNSYNC_LOG(L"TestSync(%hs, %hs)", ToString(WeakHasher), ToString(StrongHasher));
 
@@ -314,7 +317,7 @@ TestSync(EWeakHashAlgorithmID WeakHasher, EStrongHashAlgorithmID StrongHasher)
 
 		FMemReader SourceReader(Source);
 		FMemReader BaseReader(Base);
-		auto	   Target = BuildTargetBuffer(SourceReader, BaseReader, NeedBlocks, StrongHasher);
+		auto	   Target = BuildTargetBuffer(SourceReader, BaseReader, NeedBlocks, BuildParams);
 
 		UNSYNC_ASSERT(Target.Size() == Source.Size());
 		UNSYNC_ASSERT(!memcmp(Target.Data(), Source.Data(), Source.Size()));
@@ -329,7 +332,7 @@ TestSync(EWeakHashAlgorithmID WeakHasher, EStrongHashAlgorithmID StrongHasher)
 
 		UNSYNC_ASSERT(!IsSynchronized(NeedBlocks, SourceBlocks));
 
-		auto Target = BuildTargetBuffer(Source.Data(), Source.Size(), Base.Data(), Base.Size(), NeedBlocks, StrongHasher);
+		auto Target = BuildTargetBuffer(Source.Data(), Source.Size(), Base.Data(), Base.Size(), NeedBlocks, BuildParams);
 
 		UNSYNC_ASSERT(Target.Size() == Source.Size());
 		UNSYNC_ASSERT(!memcmp(Target.Data(), Source.Data(), Source.Size()));
@@ -344,7 +347,7 @@ TestSync(EWeakHashAlgorithmID WeakHasher, EStrongHashAlgorithmID StrongHasher)
 
 		UNSYNC_ASSERT(!IsSynchronized(NeedBlocks, SourceBlocks));
 
-		auto Target = BuildTargetBuffer(Source.Data(), Source.Size(), Base.Data(), Base.Size(), NeedBlocks, StrongHasher);
+		auto Target = BuildTargetBuffer(Source.Data(), Source.Size(), Base.Data(), Base.Size(), NeedBlocks, BuildParams);
 
 		UNSYNC_ASSERT(Target.Size() == Source.Size());
 		UNSYNC_ASSERT(!memcmp(Target.Data(), Source.Data(), Source.Size()));
@@ -359,7 +362,7 @@ TestSync(EWeakHashAlgorithmID WeakHasher, EStrongHashAlgorithmID StrongHasher)
 
 		UNSYNC_ASSERT(!IsSynchronized(NeedBlocks, SourceBlocks));
 
-		auto Target = BuildTargetBuffer(Source.Data(), Source.Size(), Base.Data(), Base.Size(), NeedBlocks, StrongHasher);
+		auto Target = BuildTargetBuffer(Source.Data(), Source.Size(), Base.Data(), Base.Size(), NeedBlocks, BuildParams);
 
 		UNSYNC_ASSERT(Target.Size() == Source.Size());
 		UNSYNC_ASSERT(!memcmp(Target.Data(), Source.Data(), Source.Size()));
@@ -374,7 +377,7 @@ TestSync(EWeakHashAlgorithmID WeakHasher, EStrongHashAlgorithmID StrongHasher)
 
 		UNSYNC_ASSERT(!IsSynchronized(NeedBlocks, SourceBlocks));
 
-		auto Target = BuildTargetBuffer(Source.Data(), Source.Size(), Base.Data(), Base.Size(), NeedBlocks, StrongHasher);
+		auto Target = BuildTargetBuffer(Source.Data(), Source.Size(), Base.Data(), Base.Size(), NeedBlocks, BuildParams);
 
 		UNSYNC_ASSERT(Target.Size() == Source.Size());
 		UNSYNC_ASSERT(!memcmp(Target.Data(), Source.Data(), Source.Size()));
@@ -389,7 +392,7 @@ TestSync(EWeakHashAlgorithmID WeakHasher, EStrongHashAlgorithmID StrongHasher)
 
 		UNSYNC_ASSERT(!IsSynchronized(NeedBlocks, SourceBlocks));
 
-		auto Target = BuildTargetBuffer(Source.Data(), Source.Size(), Base.Data(), Base.Size(), NeedBlocks, StrongHasher);
+		auto Target = BuildTargetBuffer(Source.Data(), Source.Size(), Base.Data(), Base.Size(), NeedBlocks, BuildParams);
 
 		UNSYNC_ASSERT(Target.Size() == Source.Size());
 		UNSYNC_ASSERT(!memcmp(Target.Data(), Source.Data(), Source.Size()));
@@ -404,7 +407,7 @@ TestSync(EWeakHashAlgorithmID WeakHasher, EStrongHashAlgorithmID StrongHasher)
 
 		UNSYNC_ASSERT(IsSynchronized(NeedBlocks, SourceBlocks));
 
-		auto Target = BuildTargetBuffer(Source.Data(), Source.Size(), Base.Data(), Base.Size(), NeedBlocks, StrongHasher);
+		auto Target = BuildTargetBuffer(Source.Data(), Source.Size(), Base.Data(), Base.Size(), NeedBlocks, BuildParams);
 
 		UNSYNC_ASSERT(Target.Size() == Source.Size());
 		UNSYNC_ASSERT(!memcmp(Target.Data(), Source.Data(), Source.Size()));
@@ -417,7 +420,7 @@ TestSync(EWeakHashAlgorithmID WeakHasher, EStrongHashAlgorithmID StrongHasher)
 		auto	SourceBlocks = ComputeBlocks(Source.Data(), Source.Size(), BlockSize, Algorithm);
 		auto	NeedBlocks	 = DiffBlocks(Base.Data(), Base.Size(), BlockSize, WeakHasher, StrongHasher, SourceBlocks);
 
-		auto Target = BuildTargetBuffer(Source.Data(), Source.Size(), Base.Data(), Base.Size(), NeedBlocks, StrongHasher);
+		auto Target = BuildTargetBuffer(Source.Data(), Source.Size(), Base.Data(), Base.Size(), NeedBlocks, BuildParams);
 
 		UNSYNC_ASSERT(Target.Size() == Source.Size());
 		UNSYNC_ASSERT(!memcmp(Target.Data(), Source.Data(), Source.Size()));
@@ -444,8 +447,11 @@ TestBuildTarget(EWeakHashAlgorithmID WeakHasher, EStrongHashAlgorithmID StrongHa
 	auto NeedBlocks	  = DiffBlocks(Base.Data(), Base.Size(), BlockSize, WeakHasher, StrongHasher, SourceBlocks);
 
 	UNSYNC_ASSERT(!IsSynchronized(NeedBlocks, SourceBlocks));
+	
+	FBuildTargetParams BuildParams;
+	BuildParams.StrongHasher = StrongHasher;
 
-	auto Target = BuildTargetBuffer(Source.Data(), Source.Size(), Base.Data(), Base.Size(), NeedBlocks, StrongHasher);
+	auto Target = BuildTargetBuffer(Source.Data(), Source.Size(), Base.Data(), Base.Size(), NeedBlocks, BuildParams);
 
 	UNSYNC_ASSERT(Target.Size() == Source.Size());
 	UNSYNC_ASSERT(!memcmp(Target.Data(), Source.Data(), Source.Size()));
