unsync - Iteration on pack command

* Store directory manfiest chunks in the pack
* Generate a snapshot file to reconstruct full manifest from chunks

diff --git a/Engine/Source/Programs/Unsync/Private/UnsyncCmdPack.cpp b/Engine/Source/Programs/Unsync/Private/UnsyncCmdPack.cpp
--- a/Engine/Source/Programs/Unsync/Private/UnsyncCmdPack.cpp
+++ b/Engine/Source/Programs/Unsync/Private/UnsyncCmdPack.cpp
@@ -1,11 +1,11 @@
 // Copyright Epic Games, Inc. All Rights Reserved.
 
 #include "UnsyncCmdPack.h"
-#include "UnsyncFile.h"
-#include "UnsyncThread.h"
-#include "UnsyncSerialization.h"
 #include "UnsyncCompression.h"
+#include "UnsyncFile.h"
 #include "UnsyncHashTable.h"
+#include "UnsyncSerialization.h"
+#include "UnsyncThread.h"
 
 #include <atomic>
 
@@ -41,7 +41,7 @@ BuildP4HaveSet(const FPath& Root, std::string_view P4HaveDataUtf8, FDirectoryMan
 {
 	auto Callback = [&Result, &Root](std::string_view LineView)
 	{
-		if (LineView.starts_with("---")) // p4 diagnostic data
+		if (LineView.starts_with("---"))  // p4 diagnostic data
 		{
 			return;
 		}
@@ -83,20 +83,129 @@ BuildP4HaveSet(const FPath& Root, std::string_view P4HaveDataUtf8, FDirectoryMan
 
 struct FPackIndexEntry
 {
-	FHash160 Hash	 = {};
+	FHash128 BlockHash		= {};
+	FHash128 CompressedHash = {};
+	uint32	 Offset			= 0;
 	uint32	 CompressedSize = 0;
-	uint64	 Offset  = 0;
 };
-static_assert(sizeof(FPackIndexEntry) == 32);
+static_assert(sizeof(FPackIndexEntry) == 40);
+
+inline void
+AddHash(uint64* Accumulator, const FHash128& Hash)
+{
+	uint64 BlockHashParts[2];
+	memcpy(BlockHashParts, &Hash, sizeof(FHash128));
+	Accumulator[0] += BlockHashParts[0];
+	Accumulator[1] += BlockHashParts[1];
+}
+
+inline FHash128
+MakeHashFromParts(uint64* Parts)
+{
+	FHash128 Result;
+	memcpy(&Result, Parts, sizeof(Result));
+	return Result;
+}
+
+struct FPackWriteContext
+{
+	static constexpr uint64 MaxPackFileSize = 1_GB;
+
+	FPackWriteContext(const FPath& InOutputRoot) : OutputRoot(InOutputRoot) { Reset(); }
+	~FPackWriteContext() { FinishPack(); }
+
+	void AddBlock(const FGenericBlock& Block, FHash128 CompressedHash, FBufferView CompressedData)
+	{
+		std::lock_guard<std::mutex> LockGuard(Mutex);
+
+		UNSYNC_ASSERT(CompressedData.Size <= MaxPackFileSize);
+
+		if (PackBuffer.Size() + CompressedData.Size > MaxPackFileSize)
+		{
+			FinishPack();
+		}
+
+		FPackIndexEntry IndexEntry;
+		IndexEntry.BlockHash	  = Block.HashStrong.ToHash128();
+		IndexEntry.CompressedHash = CompressedHash;
+		IndexEntry.Offset		  = CheckedNarrow(PackBuffer.Size());
+		IndexEntry.CompressedSize = CheckedNarrow(CompressedData.Size);
+
+		IndexEntries.push_back(IndexEntry);
+		PackBuffer.Append(CompressedData);
+
+		AddHash(IndexFileHashSum, IndexEntry.BlockHash);
+	}
+
+	void FinishPack()
+	{
+		if (IndexEntries.empty())
+		{
+			return;
+		}
+
+		FHash128	BlockHash128 = MakeHashFromParts(IndexFileHashSum);
+		std::string OutputId	 = HashToHexString(BlockHash128);
+
+		FPath FinalPackFilename	 = OutputRoot / (OutputId + ".unsync_pack");
+		FPath FinalIndexFilename = OutputRoot / (OutputId + ".unsync_index");
+
+		UNSYNC_LOG(L"Saving new pack: %hs", OutputId.c_str());
+
+		if (!WriteBufferToFile(FinalPackFilename, PackBuffer, EFileMode::CreateWriteOnly))
+		{
+			UNSYNC_FATAL(L"Failed to write pack file '%ls'", FinalPackFilename.wstring().c_str());
+		}
+
+		const uint8* IndexData	   = reinterpret_cast<const uint8*>(IndexEntries.data());
+		uint64		 IndexDataSize = sizeof(IndexEntries[0]) * IndexEntries.size();
+		if (!WriteBufferToFile(FinalIndexFilename, IndexData, IndexDataSize, EFileMode::CreateWriteOnly))
+		{
+			UNSYNC_FATAL(L"Failed to write index file '%ls'", FinalIndexFilename.wstring().c_str());
+		}
+
+		Reset();
+	}
+
+private:
+	void Reset()
+	{
+		PackBuffer.Reserve(MaxPackFileSize);
+		PackBuffer.Clear();
+		IndexEntries.clear();
+
+		IndexFileHashSum[0] = 0;
+		IndexFileHashSum[1] = 0;
+	}
+
+	std::mutex Mutex;
+
+	// Independent sums of low and high 32 bits of all seen block hashes.
+	// Used to generate a stable hash while allowing out-of-order block processing.
+	uint64 IndexFileHashSum[2] = {};
+
+	FBuffer						 PackBuffer;
+	std::vector<FPackIndexEntry> IndexEntries;
+
+	FPath OutputRoot;
+};
 
-int32 CmdPack(const FCmdPackOptions& Options)
+static bool
+EnsureDirectoryExists(const FPath& Path)
+{
+	return (PathExists(Path) && IsDirectory(Path)) || CreateDirectories(Path);
+}
+
+int32
+CmdPack(const FCmdPackOptions& Options)
 {
 	const FFileAttributes RootAttrib = GetFileAttrib(Options.RootPath);
 
 	const FPath InputRoot	 = Options.RootPath;
 	const FPath ManifestRoot = InputRoot / ".unsync";
+	const FPath PackRoot	 = ManifestRoot / "pack";  // TODO: override this via command line
 
-	UNSYNC_LOG(L"Generating package for directory '%ls' ...", InputRoot.wstring().c_str());
+	UNSYNC_LOG(L"Generating package for directory '%ls'", InputRoot.wstring().c_str());
 	UNSYNC_LOG_INDENT;
 
 	if (!RootAttrib.bValid)
@@ -112,43 +221,33 @@ int32 CmdPack(const FCmdPackOptions& Options)
 	}
 
 	// TODO: allow explicit output path
-	const bool bOutputDirectoryCreated = (PathExists(ManifestRoot) && IsDirectory(ManifestRoot)) || CreateDirectories(ManifestRoot);
-	if (!bOutputDirectoryCreated)
-	{
-		UNSYNC_ERROR(L"Failed to create output directory '%ls'", ManifestRoot.wstring().c_str());
-		return -1;
-	}
-
-	FPath		TempOutputPackFilename = ManifestRoot / "blocks.bin.tmp";
-	FNativeFile PackFile(TempOutputPackFilename, EFileMode::CreateWriteOnly);
-	if (!PackFile.IsValid())
 	{
-		UNSYNC_ERROR(L"Failed open pack output file '%ls'", TempOutputPackFilename.wstring().c_str());
-		return -1;
-	}
+		if (!EnsureDirectoryExists(ManifestRoot))
+		{
+			UNSYNC_ERROR(L"Failed to create manifest output directory '%ls'", ManifestRoot.wstring().c_str());
+			return -1;
+		}
 
-	FPath		TempOutputIndexFilename = ManifestRoot / "blocks.idx.tmp";
-	FNativeFile IndexFile(TempOutputIndexFilename, EFileMode::CreateWriteOnly);
-	if (!IndexFile.IsValid())
-	{
-		UNSYNC_ERROR(L"Failed open index output file '%ls'", TempOutputIndexFilename.wstring().c_str());
-		return -1;
+		if (!EnsureDirectoryExists(PackRoot))
+		{
+			UNSYNC_ERROR(L"Failed to create pack output directory '%ls'", PackRoot.wstring().c_str());
+			return -1;
+		}
 	}
 
-	THashSet<FHash160> SeenMacroBlockHashSet;
-	UNSYNC_LOG(L"Loading existing block packs ...");
+	THashSet<FHash128> SeenBlockHashSet;
+	UNSYNC_LOG(L"Loading block database");
 	{
 		UNSYNC_LOG_INDENT;
-		FPath ExistingIndexSearchPath = ManifestRoot; // TODO: could allow overriding this
-		FPath ExpectedExtension		  = FPath(".idx");
-		for (const std::filesystem::directory_entry& Dir : RecursiveDirectoryScan(ExistingIndexSearchPath))
+		FPath ExpectedExtension = FPath(".unsync_index");
+		for (const std::filesystem::directory_entry& Dir : RecursiveDirectoryScan(PackRoot))
 		{
 			if (!Dir.is_regular_file())
 			{
 				continue;
 			}
 
-			const FPath& FilePath = Dir.path();
+			const FPath&	FilePath = Dir.path();
 			FPathStringView FilePathView(FilePath.native());
 			if (!FilePathView.ends_with(ExpectedExtension.native()))
 			{
@@ -159,7 +258,7 @@ int32 CmdPack(const FCmdPackOptions& Options)
 			uint64	NumEntries		= ExistingEntries.Size() / sizeof(FPackIndexEntry);
 			for (const FPackIndexEntry& Entry : MakeView(reinterpret_cast<FPackIndexEntry*>(ExistingEntries.Data()), NumEntries))
 			{
-				SeenMacroBlockHashSet.insert(Entry.Hash);
+				SeenBlockHashSet.insert(Entry.BlockHash);
 			}
 		}
 	}
@@ -195,7 +294,7 @@ int32 CmdPack(const FCmdPackOptions& Options)
 
 		UNSYNC_LOG(L"Loaded entries from p4 manifest: %llu", llu(DirectoryManifest.Files.size()));
 
-		UNSYNC_LOG(L"Reading file attributes ...");
+		UNSYNC_LOG(L"Reading file attributes");
 		auto UpdateFileMetadata = [](std::pair<const std::wstring, FFileManifest>& It)
 		{
 			FFileAttributes Attrib = GetFileAttrib(It.second.CurrentPath);
@@ -219,26 +318,24 @@ int32 CmdPack(const FCmdPackOptions& Options)
 
 	UNSYNC_LOG(L"Found files: %llu", llu(DirectoryManifest.Files.size()));
 
-	FComputeBlocksParams BlockParams;
-	BlockParams.Algorithm = Options.Algorithm;
-	BlockParams.BlockSize = Options.BlockSize;
-	BlockParams.bNeedMacroBlocks = true;
-
-	std::atomic<uint64> PackFileOffset;
-	std::atomic<uint64> IndexFileOffset;
 	std::atomic<uint64> ProcessedRawBytes;
+	std::atomic<uint64> CompressedBytes;
+
+	FPackWriteContext PackWriter(PackRoot);
 
-	std::atomic<uint64> IndexFileHashSum[2];
+	FThreadLogConfig LogConfig;
 
 	std::mutex Mutex;
-	BlockParams.OnMacroBlockGenerated =
-		[&Mutex, &PackFile, &PackFileOffset, &IndexFile, &IndexFileOffset, &SeenMacroBlockHashSet, &IndexFileHashSum, &ProcessedRawBytes](
-			const FGenericBlock& Block,
-			FBufferView			 Data)
+
+	FOnBlockGenerated OnBlockGenerated =
+		[&Mutex, &PackWriter, &SeenBlockHashSet, &ProcessedRawBytes, &CompressedBytes, &LogConfig](const FGenericBlock& Block,
+																								   FBufferView			Data)
 	{
+		FThreadLogConfig::FScope LogConfigScope(LogConfig);
+
 		{
 			std::lock_guard<std::mutex> LockGuard(Mutex);
-			if (!SeenMacroBlockHashSet.insert(Block.HashStrong.ToHash160()).second)
+			if (!SeenBlockHashSet.insert(Block.HashStrong.ToHash128()).second)
 			{
 				return;
 			}
@@ -248,33 +345,31 @@ int32 CmdPack(const FCmdPackOptions& Options)
 		FIOBuffer	 CompressedData		  = FIOBuffer::Alloc(MaxCompressedSize, L"PackBlock");
 		uint64		 ActualCompressedSize = CompressInto(Data, CompressedData.GetMutBufferView(), 9);
 
-		if (ActualCompressedSize)
+		if (!ActualCompressedSize)
 		{
-			ProcessedRawBytes += Data.Size;
+			UNSYNC_FATAL(L"Failed to compress file block");
+		}
+		CompressedData.SetDataRange(0, ActualCompressedSize);
 
-			CompressedData.SetDataRange(0, ActualCompressedSize);
+		ProcessedRawBytes += Block.Size;
+		CompressedBytes += ActualCompressedSize;
 
-			uint64 PackWriteOffset = PackFileOffset.fetch_add(ActualCompressedSize);
-			PackFile.Write(CompressedData.GetData(), PackWriteOffset, ActualCompressedSize);
+		FHash128 CompressedHash = HashBlake3Bytes<FHash128>(CompressedData.GetData(), ActualCompressedSize);
 
-			FPackIndexEntry IndexEntry;
-			IndexEntry.Hash			  = Block.HashStrong.ToHash160();
-			IndexEntry.CompressedSize = CheckedNarrow<uint64>(ActualCompressedSize);
-			IndexEntry.Offset		  = PackWriteOffset;
+		PackWriter.AddBlock(Block, CompressedHash, CompressedData.GetBufferView());
+	};
 
-			uint64 IndexWriteOffset = IndexFileOffset.fetch_add(sizeof(IndexEntry));
-			IndexFile.Write(&IndexEntry, IndexWriteOffset, sizeof(IndexEntry));
+	FComputeBlocksParams BlockParams;
+	BlockParams.Algorithm		 = Options.Algorithm;
+	BlockParams.BlockSize		 = Options.BlockSize;
 
-			FHash128 BlockHash128 = Block.HashStrong.ToHash128();
-			uint64	 BlockHashParts[2];
-			memcpy(BlockHashParts, &BlockHash128, sizeof(BlockHash128));
-			IndexFileHashSum[0] += BlockHashParts[0];
-			IndexFileHashSum[1] += BlockHashParts[1];
-		}
-	};
+	// TODO: threading makes pack files non-deterministic...
+	// Perhaps a strictly ordered parallel pipeline mechanism could be implemented?
+	BlockParams.bAllowThreading	 = true;
+	BlockParams.OnBlockGenerated = OnBlockGenerated;
 
 	{
-		UNSYNC_LOG(L"Loading previous manifest ... ")
+		UNSYNC_LOG(L"Loading previous manifest ")
 		FDirectoryManifest OldManifest;
 		if (PathExists(DirectoryManifestPath) && LoadDirectoryManifest(OldManifest, InputRoot, DirectoryManifestPath))
 		{
@@ -307,71 +402,96 @@ int32 CmdPack(const FCmdPackOptions& Options)
 		}
 	}
 
-	UNSYNC_LOG(L"Computing file blocks ...");
+	UNSYNC_LOG(L"Computing file blocks");
 	UpdateDirectoryManifestBlocks(DirectoryManifest, InputRoot, BlockParams);
 
-	const uint64 CompressedSize = PackFileOffset;
+	uint64 ManifestUniqueBytes	   = 0;
+	uint64 ManifestCompressedBytes = 0;
 
-	if (!GDryRun)
-	{
-		UNSYNC_LOG(L"Saving directory manifest '%ls'", DirectoryManifestPath.wstring().c_str());
-		SaveDirectoryManifest(DirectoryManifest, DirectoryManifestPath);
-	}
+	UNSYNC_LOG(L"Saving directory manifest");
 
-	uint64 SourceSize = 0;
-	for (const auto& It : DirectoryManifest.Files)
 	{
-		SourceSize += It.second.Size;
-	}
-
-	const uint64 NumSourceFiles = DirectoryManifest.Files.size();
-	UNSYNC_LOG(L"Source files: %llu", llu(NumSourceFiles));
-	UNSYNC_LOG(L"Source size: %llu bytes (%.2f MB)", llu(SourceSize), SizeMb(SourceSize));
-	UNSYNC_LOG(L"New data size: %llu bytes (%.2f MB)", llu(ProcessedRawBytes), SizeMb(ProcessedRawBytes));
-	UNSYNC_LOG(L"Compressed size: %llu bytes (%.2f MB), %.0f%%",
-			   llu(CompressedSize),
-			   SizeMb(CompressedSize),
-			   ProcessedRawBytes > 0 ? (100.0 * double(CompressedSize) / double(ProcessedRawBytes)) : 0);
+		FBuffer			 ManifestBuffer;
+		FVectorStreamOut ManifestStream(ManifestBuffer);
+		bool			 bManifestSerialized = SaveDirectoryManifest(DirectoryManifest, ManifestStream);
+		if (!bManifestSerialized)
+		{
+			UNSYNC_FATAL(L"Failed to serialize directory manifest to memory");
+			return -1;
+		}
 
-	PackFile.Close();
-	IndexFile.Close();
+		std::vector<FGenericBlock> ManifestBlocks;
+		FOnBlockGenerated		   OnManifestBlockGenerated =
+			[&OnBlockGenerated, &Mutex, &ManifestBlocks](const FGenericBlock& Block, FBufferView Data)
+		{
+			{
+				std::lock_guard<std::mutex> LockGuard(Mutex);
+				ManifestBlocks.push_back(Block);
+			}
+			OnBlockGenerated(Block, Data);
+		};
 
-	std::error_code ErrorCode;
+		ManifestUniqueBytes -= ProcessedRawBytes.load();
+		ManifestCompressedBytes -= CompressedBytes.load();
 
-	if (PackFileOffset == 0)
-	{
-		UNSYNC_LOG(L"No new blocks found");
-		FileRemove(TempOutputPackFilename, ErrorCode);
-		FileRemove(TempOutputIndexFilename, ErrorCode);
-	}
-	else
-	{
-		uint64 IndexHashParts[2];
-		IndexHashParts[0] = IndexFileHashSum[0];
-		IndexHashParts[1] = IndexFileHashSum[1];
-		FHash128 BlockHash128;
-		memcpy(&BlockHash128, IndexHashParts, sizeof(BlockHash128));
-		std::string OutputId = HashToHexString(BlockHash128);
+		BlockParams.OnBlockGenerated = OnManifestBlockGenerated;
+		FMemReader ManifestDataReader(ManifestBuffer);
+		ComputeBlocks(ManifestDataReader, BlockParams);
 
-		FPath FinalPackFilename	 = ManifestRoot / (OutputId + ".bin");
-		FPath FinalIndexFilename = ManifestRoot / (OutputId + ".idx");
+		if (!WriteBufferToFile(DirectoryManifestPath, ManifestBuffer))
+		{
+			UNSYNC_FATAL(L"Failed to save directory manifest to file '%ls'", DirectoryManifestPath.wstring().c_str());
+			return 1;
+		}
 
-		UNSYNC_LOG(L"Saving new pack: %hs", OutputId.c_str());
+		std::sort(ManifestBlocks.begin(), ManifestBlocks.end(), FGenericBlock::FCompareByOffset());
 
-		if (!FileRename(TempOutputPackFilename, FinalPackFilename, ErrorCode))
+		for (const FGenericBlock& Block : ManifestBlocks)
 		{
-			UNSYNC_ERROR(L"Failed to rename temporary pack file to '%ls'", FinalPackFilename.wstring().c_str());
-			return 1;
+			UNSYNC_ASSERT(SeenBlockHashSet.find(Block.HashStrong.ToHash128()) != SeenBlockHashSet.end());
 		}
 
-		if (!FileRename(TempOutputIndexFilename, FinalIndexFilename, ErrorCode))
+		FBufferView ManifestBlocksBuffer;
+		ManifestBlocksBuffer.Data = reinterpret_cast<const uint8*>(ManifestBlocks.data());
+		ManifestBlocksBuffer.Size = sizeof(ManifestBlocks[0]) * ManifestBlocks.size();
+
+		FHash128 ManifestBlocksBufferHash = HashBlake3Bytes<FHash128>(ManifestBlocksBuffer.Data, ManifestBlocksBuffer.Size);
+
+		std::string SnapshotId	 = HashToHexString(ManifestBlocksBufferHash);  // TODO: allow overriding this from command line
+		FPath		SnapshotPath = ManifestRoot / (SnapshotId + ".unsync_snapshot");
+
+		UNSYNC_LOG(L"Writing snapshot: %hs", SnapshotId.c_str());
+
+		bool bSnapshotWritten = WriteBufferToFile(SnapshotPath, ManifestBlocksBuffer.Data, ManifestBlocksBuffer.Size);
+		if (!bSnapshotWritten)
 		{
-			UNSYNC_ERROR(L"Failed to rename temporary index file to '%ls'", FinalIndexFilename.wstring().c_str());
-			return 1;
+			UNSYNC_FATAL(L"Failed to write snapthot file '%ls'", SnapshotPath.wstring().c_str());
 		}
+
+		ManifestUniqueBytes += ProcessedRawBytes.load();
+		ManifestCompressedBytes += CompressedBytes.load();
+	}
+
+	PackWriter.FinishPack();
+
+	uint64 SourceSize = 0;
+	for (const auto& It : DirectoryManifest.Files)
+	{
+		SourceSize += It.second.Size;
 	}
 
+	const uint64 NumSourceFiles = DirectoryManifest.Files.size();
+	UNSYNC_LOG(L"Source files: %llu", llu(NumSourceFiles));
+	UNSYNC_LOG(L"Source size: %llu bytes (%.2f MB)", llu(SourceSize), SizeMb(SourceSize));
+	UNSYNC_LOG(L"Manifest unique data size: %llu bytes (%.2f MB)", llu(ManifestUniqueBytes), SizeMb(ManifestUniqueBytes));
+	UNSYNC_LOG(L"Manifest unique compressed size: %llu bytes (%.2f MB)", llu(ManifestCompressedBytes), SizeMb(ManifestCompressedBytes));
+	UNSYNC_LOG(L"New data size: %llu bytes (%.2f MB)", llu(ProcessedRawBytes), SizeMb(ProcessedRawBytes));
+	UNSYNC_LOG(L"Compressed size: %llu bytes (%.2f MB), %.0f%%",
+			   llu(CompressedBytes.load()),
+			   SizeMb(CompressedBytes.load()),
+			   ProcessedRawBytes > 0 ? (100.0 * double(CompressedBytes.load()) / double(ProcessedRawBytes)) : 0);
+
 	return 0;
 }
 
-}
+}  // namespace unsync
diff --git a/Engine/Source/Programs/Unsync/Private/UnsyncCore.h b/Engine/Source/Programs/Unsync/Private/UnsyncCore.h
--- a/Engine/Source/Programs/Unsync/Private/UnsyncCore.h
+++ b/Engine/Source/Programs/Unsync/Private/UnsyncCore.h
@@ -124,6 +124,8 @@ struct FComputeBlocksParams
 	// Callbacks may be called from worker threads
 	FOnBlockGenerated OnBlockGenerated;
 	FOnBlockGenerated OnMacroBlockGenerated;
+
+	bool bAllowThreading = true;
 };
 
 struct FComputeBlocksResult
diff --git a/Engine/Source/Programs/Unsync/Private/UnsyncManifest.cpp b/Engine/Source/Programs/Unsync/Private/UnsyncManifest.cpp
--- a/Engine/Source/Programs/Unsync/Private/UnsyncManifest.cpp
+++ b/Engine/Source/Programs/Unsync/Private/UnsyncManifest.cpp
@@ -245,7 +245,14 @@ UpdateDirectoryManifestBlocks(FDirectoryManifest& Result, const FPath& Root, con
 			Semaphore.Release();
 		};
 
-		TaskGroup.run(BlockTask);
+		if (Params.bAllowThreading)
+		{
+			TaskGroup.run(BlockTask);
+		}
+		else
+		{
+			BlockTask();
+		}
 	}
 
 	TaskGroup.wait();
diff --git a/Engine/Source/Programs/Unsync/Private/UnsyncProtocol.h b/Engine/Source/Programs/Unsync/Private/UnsyncProtocol.h
--- a/Engine/Source/Programs/Unsync/Private/UnsyncProtocol.h
+++ b/Engine/Source/Programs/Unsync/Private/UnsyncProtocol.h
@@ -57,6 +57,11 @@ struct TBlock
 	uint32		Size	   = 0;
 	uint32		HashWeak   = 0;
 	StrongHashT HashStrong = {};
+
+	struct FCompareByOffset
+	{
+		bool operator()(const TBlock<StrongHashT>& A, const TBlock<StrongHashT>& B) const { return A.Offset < B.Offset; }
+	};
 };
 
 using FBlock128 = TBlock<FHash128>;
diff --git a/Engine/Source/Programs/Unsync/Private/UnsyncThread.h b/Engine/Source/Programs/Unsync/Private/UnsyncThread.h
--- a/Engine/Source/Programs/Unsync/Private/UnsyncThread.h
+++ b/Engine/Source/Programs/Unsync/Private/UnsyncThread.h
@@ -3,6 +3,7 @@
 #pragma once
 
 #include "UnsyncUtil.h"
+#include "UnsyncLog.h"
 
 UNSYNC_THIRD_PARTY_INCLUDES_START
 #include <atomic>
@@ -67,6 +68,29 @@ struct FThreadElectScope
 	operator bool() const { return bValue; }
 };
 
+struct FThreadLogConfig
+{
+	FThreadLogConfig() : ParentThreadIndent(GLogIndent), bParentThreadVerbose(GLogVerbose) {}
+
+
+	uint32				ParentThreadIndent;
+	bool				bParentThreadVerbose;
+	std::atomic<uint64> NumActiveVerboseLogThreads = {};
+
+	struct FScope
+	{
+		FScope(FThreadLogConfig& Parent)
+		: AllowVerbose(Parent.NumActiveVerboseLogThreads, Parent.bParentThreadVerbose)
+		, VerboseScope(AllowVerbose.bValue)
+		, IndentScope(Parent.ParentThreadIndent, true)
+		{
+		}
+		FThreadElectScope  AllowVerbose;
+		FLogVerbosityScope VerboseScope;
+		FLogIndentScope	   IndentScope;
+	};
+};
+
 void SchedulerSleep(uint32 Milliseconds);
 void SchedulerYield();
 
