unsync - Use an regular ordered map to store file list in the manifest

* Ordered file list results in manfiests that can themselves be chunked and incrementally downloaded
* Add a command line parameter to specify storage location for pack command

diff --git a/Engine/Source/Programs/Unsync/Private/UnsyncCmdPack.cpp b/Engine/Source/Programs/Unsync/Private/UnsyncCmdPack.cpp
--- a/Engine/Source/Programs/Unsync/Private/UnsyncCmdPack.cpp
+++ b/Engine/Source/Programs/Unsync/Private/UnsyncCmdPack.cpp
@@ -203,7 +203,8 @@ CmdPack(const FCmdPackOptions& Options)
 
 	const FPath InputRoot	 = Options.RootPath;
 	const FPath ManifestRoot = InputRoot / ".unsync";
-	const FPath PackRoot	 = ManifestRoot / "pack";  // TODO: override this via command line
+	const FPath StoreRoot	 = Options.StorePath.empty() ? ManifestRoot : Options.StorePath;
+	const FPath PackRoot	 = StoreRoot / "pack";
 
 	UNSYNC_LOG(L"Generating package for directory '%ls'", InputRoot.wstring().c_str());
 	UNSYNC_LOG_INDENT;
@@ -458,7 +459,7 @@ CmdPack(const FCmdPackOptions& Options)
 		FHash128 ManifestBlocksBufferHash = HashBlake3Bytes<FHash128>(ManifestBlocksBuffer.Data, ManifestBlocksBuffer.Size);
 
 		std::string SnapshotId	 = HashToHexString(ManifestBlocksBufferHash);  // TODO: allow overriding this from command line
-		FPath		SnapshotPath = ManifestRoot / (SnapshotId + ".unsync_snapshot");
+		FPath		SnapshotPath = StoreRoot / (SnapshotId + ".unsync_snapshot");
 
 		UNSYNC_LOG(L"Writing snapshot: %hs", SnapshotId.c_str());
 
diff --git a/Engine/Source/Programs/Unsync/Private/UnsyncCmdPack.h b/Engine/Source/Programs/Unsync/Private/UnsyncCmdPack.h
--- a/Engine/Source/Programs/Unsync/Private/UnsyncCmdPack.h
+++ b/Engine/Source/Programs/Unsync/Private/UnsyncCmdPack.h
@@ -9,8 +9,9 @@ namespace unsync {
 
 struct FCmdPackOptions
 {
-	FPath RootPath;
-	FPath P4HavePath; // optional
+	FPath			  RootPath;
+	FPath			  P4HavePath;  // optional
+	FPath			  StorePath;   // optional
 	uint32			  BlockSize = uint32(64_KB);
 	FAlgorithmOptions Algorithm;
 };
diff --git a/Engine/Source/Programs/Unsync/Private/UnsyncMain.cpp b/Engine/Source/Programs/Unsync/Private/UnsyncMain.cpp
--- a/Engine/Source/Programs/Unsync/Private/UnsyncMain.cpp
+++ b/Engine/Source/Programs/Unsync/Private/UnsyncMain.cpp
@@ -78,6 +78,7 @@ InnerMain(int Argc, char** Argv)
 	std::vector<std::string> QueryArgsUtf8;
 	std::string				 ScavengeRootUtf8;
 	std::string				 P4HavePathUtf8;
+	std::string				 StorePathUtf8;
 	bool					 bForceOperation	 = false;
 	bool					 bAllowInsecureTls	 = false;
 	bool					 bUseTls			 = false;
@@ -157,6 +158,9 @@ InnerMain(int Argc, char** Argv)
 	CLI::App* SubPack = Cli.add_subcommand("pack", "EXPERIMENTAL: Generate manifest for a directory and store all referenced data in a compressed pack file");
 	SubPack->add_option("Input", InputFilenameUtf8, "Input directory path")->required();
 	SubPack->add_option("--p4havefile", P4HavePathUtf8, "Use `p4 have` output from a given file to explicitly specify files included in the manifest");
+	SubPack->add_option("--store",
+						StorePathUtf8,
+						"Use this location to store pack data (default: <Input>/.unsync/pack)");
 	SubCommands.push_back(SubPack);
 
 	// Configure push
@@ -765,6 +769,7 @@ InnerMain(int Argc, char** Argv)
 
 		PackOptions.RootPath	 = InputFilename;
 		PackOptions.P4HavePath	 = NormalizeFilenameUtf8(P4HavePathUtf8);
+		PackOptions.StorePath	 = NormalizeFilenameUtf8(StorePathUtf8);
 		PackOptions.BlockSize	 = HashOrSyncBlockSize;
 		PackOptions.Algorithm	 = Algorithm;
 
diff --git a/Engine/Source/Programs/Unsync/Private/UnsyncManifest.cpp b/Engine/Source/Programs/Unsync/Private/UnsyncManifest.cpp
--- a/Engine/Source/Programs/Unsync/Private/UnsyncManifest.cpp
+++ b/Engine/Source/Programs/Unsync/Private/UnsyncManifest.cpp
@@ -139,30 +139,16 @@ ComputeManifestStableSignature(const FDirectoryManifest& Manifest)
 	UpdateHashT(Hasher, Manifest.Algorithm.WeakHashAlgorithmId);
 	UpdateHashT(Hasher, Manifest.Algorithm.StrongHashAlgorithmId);
 
-	std::vector<std::wstring_view> SortedFiles;
-	SortedFiles.reserve(Manifest.Files.size());
-
-	for (const auto& It : Manifest.Files)
-	{
-		const std::wstring& FileName = It.first;
-		SortedFiles.push_back(FileName);
-	}
-
-	std::sort(SortedFiles.begin(), SortedFiles.end());
-
 	std::string	 FileNameUtf8;
-	std::wstring FileName;
-	for (const std::wstring_view& FileNameView : SortedFiles)
+	for (const auto& FileIt : Manifest.Files)
 	{
-		FileName.clear();
-		FileName.append(FileNameView);
+		const std::wstring& FileName = FileIt.first;
+		const FFileManifest& FileManifest = FileIt.second;
 
 		// Canonical unsync file paths are utf8 with unix-style separator `/`
 		ConvertWideToUtf8(FileName, FileNameUtf8);
 		std::replace(FileNameUtf8.begin(), FileNameUtf8.end(), '\\', '/');
 
-		const FFileManifest& FileManifest = Manifest.Files.at(FileName);  // TODO: heterogeneous map
-
 		blake3_hasher_update(&Hasher, FileNameUtf8.c_str(), FileNameUtf8.length());
 
 		UpdateHashT(Hasher, FileManifest.Mtime);
@@ -179,12 +165,6 @@ ComputeManifestStableSignature(const FDirectoryManifest& Manifest)
 	return Result;
 }
 
-FHash160
-ComputeManifestStableSignature160(const FDirectoryManifest& Manifest)
-{
-	return ToHash160(ComputeManifestStableSignature(Manifest));
-}
-
 void
 UpdateDirectoryManifestBlocks(FDirectoryManifest& Result, const FPath& Root, const FComputeBlocksParams& Params)
 {
diff --git a/Engine/Source/Programs/Unsync/Private/UnsyncManifest.h b/Engine/Source/Programs/Unsync/Private/UnsyncManifest.h
--- a/Engine/Source/Programs/Unsync/Private/UnsyncManifest.h
+++ b/Engine/Source/Programs/Unsync/Private/UnsyncManifest.h
@@ -6,6 +6,8 @@
 #include "UnsyncHash.h"
 #include "UnsyncProtocol.h"
 
+#include <map>
+
 namespace unsync {
 
 struct FComputeBlocksParams;
@@ -59,7 +61,7 @@ struct FDirectoryManifest
 
 	// wide string at runtime, utf8 serialized
 	// TODO: keep paths in canonical form (utf-8, unix-style separators)
-	using FFileMap = std::unordered_map<std::wstring, FFileManifest>;
+	using FFileMap = std::map<std::wstring, FFileManifest>; // regular map to keep files in a deterministic order
 	FFileMap Files;
 
 	// runtime data
@@ -91,11 +93,9 @@ void				   LogManifestFiles(ELogLevel LogLevel, const FDirectoryManifestInfo& In
 // Files are processed in sorted order, with file names treated as utf-8.
 // This produces a relatively stable manifest key that does not depend on the serialization differences between versions.
 FHash256 ComputeManifestStableSignature(const FDirectoryManifest& Manifest);
-FHash160 ComputeManifestStableSignature160(const FDirectoryManifest& Manifest);
 
 // Computes a Blake3 hash of the serialized manifest data
 FHash256 ComputeSerializedManifestHash(const FDirectoryManifest& Manifest);
-FHash160 ComputeSerializedManifestHash160(const FDirectoryManifest& Manifest);
 
 void			   UpdateDirectoryManifestBlocks(FDirectoryManifest& Result, const FPath& Root, const FComputeBlocksParams& Params);
 FDirectoryManifest CreateDirectoryManifest(const FPath& Root, const FComputeBlocksParams& Params);
diff --git a/Engine/Source/Programs/Unsync/Private/UnsyncSerialization.cpp b/Engine/Source/Programs/Unsync/Private/UnsyncSerialization.cpp
--- a/Engine/Source/Programs/Unsync/Private/UnsyncSerialization.cpp
+++ b/Engine/Source/Programs/Unsync/Private/UnsyncSerialization.cpp
@@ -623,6 +623,20 @@ WriteSection(FVectorStreamOut& Stream, FBufferView SectionData)
 	WriteSection(Stream, SectionHeaderType::MAGIC, SectionHeaderType::VERSION, SectionData);
 }
 
+bool
+ManifestHasMacroBlocks(const FDirectoryManifest& Manifest)
+{
+	// TODO: store the macro block count in the manifest runtime data
+	for (const auto& FileIt : Manifest.Files)
+	{
+		if (!FileIt.second.MacroBlocks.empty())
+		{
+			return true;
+		}
+	}
+	return false;
+}
+
 bool
 SaveDirectoryManifest(const FDirectoryManifest& Manifest, FVectorStreamOut& Stream)
 {
@@ -658,8 +672,8 @@ SaveDirectoryManifest(const FDirectoryManifest& Manifest, FVectorStreamOut& Stre
 		WriteSection(Stream, Section);
 	}
 
+	if (ManifestHasMacroBlocks(Manifest))
 	{
-		// TODO: only save macro block section if it's valid
 		FBuffer SectionBuffer = SaveMacroBlocks(Manifest);
 		WriteSection<FMacroBlockSection>(Stream, SectionBuffer.View());
 	}
