unsync - Store file revision control IDs in the manifest, when available

diff --git a/Engine/Source/Programs/Unsync/Private/UnsyncCmdPack.cpp b/Engine/Source/Programs/Unsync/Private/UnsyncCmdPack.cpp
--- a/Engine/Source/Programs/Unsync/Private/UnsyncCmdPack.cpp
+++ b/Engine/Source/Programs/Unsync/Private/UnsyncCmdPack.cpp
@@ -70,12 +70,11 @@ BuildP4HaveSet(const FPath& Root, std::string_view P4HaveDataUtf8, FDirectoryMan
 		}
 
 		FFileManifest FileManifest;
-		FileManifest.CurrentPath = std::move(LocalPath);
+		FileManifest.CurrentPath			 = std::move(LocalPath);
+		FileManifest.RevisionControlIdentity = std::move(DepotPathUtf8);
 
 		std::wstring RelativePathStr = RelativePath.wstring();
 		Result.insert(std::make_pair(std::move(RelativePathStr), FileManifest));
-
-		UNSYNC_UNUSED(DepotPathUtf8);  // perhaps we should store this too?
 	};
 
 	ForLines(P4HaveDataUtf8, Callback);
@@ -295,7 +294,7 @@ CmdPack(const FCmdPackOptions& Options)
 
 		UNSYNC_LOG(L"Loaded entries from p4 manifest: %llu", llu(DirectoryManifest.Files.size()));
 
-		UNSYNC_LOG(L"Reading file attributes");
+		UNSYNC_LOG(L"Updating file attributes");
 		auto UpdateFileMetadata = [](std::pair<const std::wstring, FFileManifest>& It)
 		{
 			FFileAttributes Attrib = GetFileAttrib(It.second.CurrentPath);
@@ -370,35 +369,20 @@ CmdPack(const FCmdPackOptions& Options)
 	BlockParams.OnBlockGenerated = OnBlockGenerated;
 
 	{
-		UNSYNC_LOG(L"Loading previous manifest ")
+		UNSYNC_LOG(L"Loading previous manifest ");
 		FDirectoryManifest OldManifest;
 		if (PathExists(DirectoryManifestPath) && LoadDirectoryManifest(OldManifest, InputRoot, DirectoryManifestPath))
 		{
-			UNSYNC_LOG(L"Previous manifest loaded")
-
-			BlockParams.Algorithm = OldManifest.Algorithm;
-
-			FDirectoryManifest& NewManifest = DirectoryManifest;
+			UNSYNC_LOG(L"Previous manifest loaded");
 
 			// Copy file blocks from old manifest, if possible
-			for (auto& NewManifestFileEntry : NewManifest.Files)
+			if (AlgorithmOptionsCompatible(DirectoryManifest.Algorithm, OldManifest.Algorithm))
+			{
+				MoveCompatibleManifestBlocks(DirectoryManifest, std::move(OldManifest));
+			}
+			else
 			{
-				const std::wstring& FileName			 = NewManifestFileEntry.first;
-				auto				OldManifestFileEntry = OldManifest.Files.find(FileName);
-				if (OldManifestFileEntry == OldManifest.Files.end())
-				{
-					continue;
-				}
-
-				FFileManifest& NewEntry = NewManifestFileEntry.second;
-				FFileManifest& OldEntry = OldManifestFileEntry->second;
-
-				if (NewEntry.Mtime == OldEntry.Mtime && NewEntry.Size == OldEntry.Size)
-				{
-					NewEntry.Blocks		 = std::move(OldEntry.Blocks);
-					NewEntry.MacroBlocks = std::move(OldEntry.MacroBlocks);
-					NewEntry.BlockSize	 = OldEntry.BlockSize;
-				}
+				UNSYNC_LOG(L"Incremental file block generation is not possible due to algorithm options mismatch");
 			}
 		}
 	}
diff --git a/Engine/Source/Programs/Unsync/Private/UnsyncManifest.cpp b/Engine/Source/Programs/Unsync/Private/UnsyncManifest.cpp
--- a/Engine/Source/Programs/Unsync/Private/UnsyncManifest.cpp
+++ b/Engine/Source/Programs/Unsync/Private/UnsyncManifest.cpp
@@ -388,6 +388,35 @@ CreateDirectoryManifest(const FPath& Root, const FComputeBlocksParams& Params)
 	return Result;
 }
 
+void
+MoveCompatibleManifestBlocks(FDirectoryManifest& Manifest, FDirectoryManifest&& DonorManifest)
+{
+	if (!AlgorithmOptionsCompatible(Manifest.Algorithm, Manifest.Algorithm))
+	{
+		UNSYNC_ERROR(L"MoveCompatibleManifestBlocks() requires that manifests use compatible block generation settings");
+		return;
+	}
+
+	for (auto& ManifestFileEntry : Manifest.Files)
+	{
+		auto DonorManifestFileEntry = DonorManifest.Files.find(ManifestFileEntry.first);
+		if (DonorManifestFileEntry == DonorManifest.Files.end())
+		{
+			continue;
+		}
+
+		FFileManifest& ResultEntry = ManifestFileEntry.second;
+		FFileManifest& DonorEntry  = DonorManifestFileEntry->second;
+
+		if (ResultEntry.Mtime == DonorEntry.Mtime && ResultEntry.Size == DonorEntry.Size)
+		{
+			ResultEntry.Blocks		= std::move(DonorEntry.Blocks);
+			ResultEntry.MacroBlocks = std::move(DonorEntry.MacroBlocks);
+			ResultEntry.BlockSize	= DonorEntry.BlockSize;
+		}
+	}
+}
+
 FDirectoryManifest
 CreateDirectoryManifestIncremental(const FPath& Root, const FComputeBlocksParams& InParams)
 {
@@ -413,25 +442,7 @@ CreateDirectoryManifestIncremental(const FPath& Root, const FComputeBlocksParams
 	FDirectoryManifest NewManifest = CreateDirectoryManifest(Root, LightweightManifestParams);
 
 	// Copy file blocks from old manifest, if possible
-	for (auto& NewManifestFileEntry : NewManifest.Files)
-	{
-		const std::wstring& FileName			 = NewManifestFileEntry.first;
-		auto				OldManifestFileEntry = OldManifest.Files.find(FileName);
-		if (OldManifestFileEntry == OldManifest.Files.end())
-		{
-			continue;
-		}
-
-		FFileManifest& NewEntry = NewManifestFileEntry.second;
-		FFileManifest& OldEntry = OldManifestFileEntry->second;
-
-		if (NewEntry.Mtime == OldEntry.Mtime && NewEntry.Size == OldEntry.Size)
-		{
-			NewEntry.Blocks		 = std::move(OldEntry.Blocks);
-			NewEntry.MacroBlocks = std::move(OldEntry.MacroBlocks);
-			NewEntry.BlockSize	 = OldEntry.BlockSize;
-		}
-	}
+	MoveCompatibleManifestBlocks(NewManifest, std::move(OldManifest));
 
 	// Generate blocks for changed or new files
 	UpdateDirectoryManifestBlocks(NewManifest, Root, Params);
diff --git a/Engine/Source/Programs/Unsync/Private/UnsyncManifest.h b/Engine/Source/Programs/Unsync/Private/UnsyncManifest.h
--- a/Engine/Source/Programs/Unsync/Private/UnsyncManifest.h
+++ b/Engine/Source/Programs/Unsync/Private/UnsyncManifest.h
@@ -22,6 +22,7 @@ struct FFileManifest
 	FGenericBlockArray Blocks;
 	FGenericBlockArray MacroBlocks;
 	FPath			   CurrentPath;
+	std::string		   RevisionControlIdentity;
 	bool			   bReadOnly = false;
 
 	bool IsValid() const { return Mtime != 0 && Size != 0; }
@@ -68,6 +69,8 @@ struct FDirectoryManifest
 	FAlgorithmOptions Algorithm = {};
 	uint64			  Version	= 0;
 
+	bool bHasFileRevisionControl = false;
+
 	bool IsValid() const { return Version != EVersions::Invalid; }
 };
 
@@ -102,6 +105,8 @@ FDirectoryManifest CreateDirectoryManifest(const FPath& Root, const FComputeBloc
 FDirectoryManifest CreateDirectoryManifestIncremental(const FPath& Root, const FComputeBlocksParams& Params);
 bool			   LoadOrCreateDirectoryManifest(FDirectoryManifest& Result, const FPath& Root, const FComputeBlocksParams& Params);
 
+void MoveCompatibleManifestBlocks(FDirectoryManifest& Manifest, FDirectoryManifest&& DonorManifest);
+
 bool AlgorithmOptionsCompatible(const FAlgorithmOptions& A, const FAlgorithmOptions& B);
 
 }  // namespace unsync
diff --git a/Engine/Source/Programs/Unsync/Private/UnsyncProtocol.h b/Engine/Source/Programs/Unsync/Private/UnsyncProtocol.h
--- a/Engine/Source/Programs/Unsync/Private/UnsyncProtocol.h
+++ b/Engine/Source/Programs/Unsync/Private/UnsyncProtocol.h
@@ -71,10 +71,11 @@ using FBlock256 = TBlock<FHash256>;
 using FGenericBlock = TBlock<FGenericHash>;
 using FGenericBlockArray = std::vector<FGenericBlock>;
 
-static constexpr uint64 SERIALIZED_SECTION_ID_TERMINATOR		  = 0;
-static constexpr uint64 SERIALIZED_SECTION_ID_METADATA_STRING	  = 0xC6BD6CDCEEF79533ull;
-static constexpr uint64 SERIALIZED_SECTION_ID_MACRO_BLOCK		  = 0x8390AEBB745E08BCull;
-static constexpr uint64 SERIALIZED_SECTION_ID_FILE_READ_ONLY_MASK = 0x851F32ED3615F0ADull;
+static constexpr uint64 SERIALIZED_SECTION_ID_TERMINATOR			= 0;
+static constexpr uint64 SERIALIZED_SECTION_ID_METADATA_STRING		= 0xC6BD6CDCEEF79533ull;
+static constexpr uint64 SERIALIZED_SECTION_ID_MACRO_BLOCK			= 0x8390AEBB745E08BCull;
+static constexpr uint64 SERIALIZED_SECTION_ID_FILE_READ_ONLY_MASK	= 0x851F32ED3615F0ADull;
+static constexpr uint64 SERIALIZED_SECTION_ID_FILE_REVISION_CONTROL = 0X2C1C72E6B78B1B50ull;
 
 struct FSerializedSectionHeader
 {
@@ -105,6 +106,12 @@ struct FFileReadOnlyMaskSection
 	static constexpr uint64 VERSION = 1;
 };
 
+struct FFileRevisionControlSection
+{
+	static constexpr uint64 MAGIC	= SERIALIZED_SECTION_ID_FILE_REVISION_CONTROL;
+	static constexpr uint64 VERSION = 1;
+};
+
 struct FBlockFileHeader
 {
 	static constexpr uint64 MAGIC	= 0x1DCB86A5BDBA27CFull;
diff --git a/Engine/Source/Programs/Unsync/Private/UnsyncSerialization.cpp b/Engine/Source/Programs/Unsync/Private/UnsyncSerialization.cpp
--- a/Engine/Source/Programs/Unsync/Private/UnsyncSerialization.cpp
+++ b/Engine/Source/Programs/Unsync/Private/UnsyncSerialization.cpp
@@ -360,6 +360,42 @@ SaveFileReadOnlyMask(const FDirectoryManifest& Manifest)
 	return Result;
 }
 
+static FBuffer
+SaveFileRevisionControl(const FDirectoryManifest& Manifest)
+{
+	FBuffer			 Result;
+	FVectorStreamOut Writer(Result);
+
+	const uint64 NumFiles = Manifest.Files.size();
+	Writer.WriteT(NumFiles);
+	for (const auto& FileIt : Manifest.Files)
+	{
+		Writer.WriteString(FileIt.second.RevisionControlIdentity);
+	}
+
+	return Result;
+}
+
+
+std::vector<std::string>
+LoadFileRevisionControl(FIOReaderStream& Reader, FSerializedSectionHeader Header)
+{
+	std::vector<std::string> Result;
+
+	uint64 NumFiles = 0;
+	Serialize(Reader, NumFiles);
+
+	Result.resize(NumFiles);
+
+	for (uint64 FileIndex = 0; FileIndex < NumFiles; ++FileIndex)
+	{
+		Serialize(Reader, Result[FileIndex]);
+	}
+
+	return Result;
+}
+
+
 bool  // TODO: return a TResult
 LoadDirectoryManifest(FDirectoryManifest& OutManifest, const FPath& Root, FIOReaderStream& Stream)
 {
@@ -410,6 +446,7 @@ LoadDirectoryManifest(FDirectoryManifest& OutManifest, const FPath& Root, FIORea
 	std::unordered_map<FHash256, FGenericBlockArray> MacroBlocks;
 
 	FBuffer FileReadOnlyMask;
+	std::vector<std::string> FileRevisions; // TODO: use linear arena to store all strings for a manifest
 
 	if (Version >= FDirectoryManifest::EVersions::V7_OptionalSections)
 	{
@@ -450,6 +487,12 @@ LoadDirectoryManifest(FDirectoryManifest& OutManifest, const FPath& Root, FIORea
 							FileReadOnlyMask = LoadFileReadOnlyMask(Stream, SectionHeader);
 							break;
 						}
+					case SERIALIZED_SECTION_ID_FILE_REVISION_CONTROL:
+						{
+							FileRevisions = LoadFileRevisionControl(Stream, SectionHeader);
+							OutManifest.bHasFileRevisionControl = true;
+							break;
+						}
 					case SERIALIZED_SECTION_ID_TERMINATOR:
 						bDoneLoadingOptionalSections = true;
 						break;
@@ -539,6 +582,11 @@ LoadDirectoryManifest(FDirectoryManifest& OutManifest, const FPath& Root, FIORea
 				FileManifest.bReadOnly = BitArrayGet(FileReadOnlyMask.Data(), FileIndex);
 			}
 
+			if (OutManifest.bHasFileRevisionControl)
+			{
+				FileManifest.RevisionControlIdentity = std::move(FileRevisions[FileIndex]);
+			}
+
 			// Store output
 			OutManifest.Files[Filename] = std::move(FileManifest);
 		}
@@ -683,6 +731,12 @@ SaveDirectoryManifest(const FDirectoryManifest& Manifest, FVectorStreamOut& Stre
 		WriteSection<FFileReadOnlyMaskSection>(Stream, SectionBuffer.View());
 	}
 
+	if (Manifest.bHasFileRevisionControl)
+	{
+		FBuffer SectionBuffer = SaveFileRevisionControl(Manifest);
+		WriteSection<FFileRevisionControlSection>(Stream, SectionBuffer.View());
+	}
+
 	// End with the terminator section (default-constructed);
 	FSerializedSectionHeader TerminatorSection;
 	Serialize(Stream, TerminatorSection);
diff --git a/Engine/Source/Programs/Unsync/Private/UnsyncUtil.cpp b/Engine/Source/Programs/Unsync/Private/UnsyncUtil.cpp
--- a/Engine/Source/Programs/Unsync/Private/UnsyncUtil.cpp
+++ b/Engine/Source/Programs/Unsync/Private/UnsyncUtil.cpp
@@ -136,10 +136,12 @@ ConvertUtf8ToWide(std::string_view StringUtf8)
 
 	if (IsTrivialAsciiString(StringUtf8))
 	{
-		Result.reserve(StringUtf8.length());
+		Result.resize(StringUtf8.length());
+		wchar_t* ResultChars = Result.data();
 		for (char c : StringUtf8)
 		{
-			Result.push_back((wchar_t)c);
+			*ResultChars = (wchar_t)c;
+			++ResultChars;
 		}
 	}
 	else
@@ -159,10 +161,12 @@ ConvertWideToUtf8(std::wstring_view StringWide, std::string& Result)
 
 	if (IsTrivialAsciiString(StringWide))
 	{
-		Result.reserve(StringWide.length());
+		Result.resize(StringWide.length());
+		char* ResultChars = Result.data();
 		for (wchar_t wc : StringWide)
 		{
-			Result.push_back((char)wc);
+			*ResultChars = (char)wc;
+			++ResultChars;
 		}
 	}
 	else
