unsync - Add option to run `p4 have` during pack process instead of reading output from file

* Fix saving revision control data in the manifest during pack command
* Reuse file attributes from existing manifest when revision control identity matches

diff --git a/Engine/Source/Programs/Unsync/Private/UnsyncCmdPack.cpp b/Engine/Source/Programs/Unsync/Private/UnsyncCmdPack.cpp
--- a/Engine/Source/Programs/Unsync/Private/UnsyncCmdPack.cpp
+++ b/Engine/Source/Programs/Unsync/Private/UnsyncCmdPack.cpp
@@ -7,6 +7,8 @@
 #include "UnsyncSerialization.h"
 #include "UnsyncThread.h"
 
+#include <stdio.h>
+#include <stdlib.h>
 #include <atomic>
 
 namespace unsync {
@@ -195,6 +197,60 @@ EnsureDirectoryExists(const FPath& Path)
 	return (PathExists(Path) && IsDirectory(Path)) || CreateDirectories(Path);
 }
 
+static int32
+RunSubprocess(const char* Command, const FPath& WorkingDirectory, std::string& StdOutBuffer)
+{
+#if UNSYNC_PLATFORM_WINDOWS
+
+	char TempBuffer[65536];
+
+	wchar_t* PrevWD = _wgetcwd(nullptr, 0);
+
+	if (!PrevWD)
+	{
+		UNSYNC_ERROR(L"Failed to get current working directory")
+		return -1;
+	}
+
+	int32 ErrorCode = _wchdir(WorkingDirectory.native().c_str());
+
+	FILE* Pipe = _popen(Command, "rt");
+	if (Pipe)
+	{
+		for (;;)
+		{
+			size_t ReadSize = fread(TempBuffer, 1, sizeof(TempBuffer), Pipe);
+			if (ReadSize != 0)
+			{
+				StdOutBuffer.append(TempBuffer, ReadSize);
+			}
+
+			if (feof(Pipe))
+			{
+				break;
+			}
+		}
+
+		ErrorCode = _pclose(Pipe);
+	}
+	else
+	{
+		ErrorCode = errno;
+	}
+
+	_wchdir(PrevWD);
+	free(PrevWD);
+
+	return ErrorCode;
+
+#else  // UNSYNC_PLATFORM_WINDOWS
+
+	UNSYNC_FATAL(L"RunSubprocess() is not implemented");
+	return -1;
+
+#endif	// UNSYNC_PLATFORM_WINDOWS
+}
+
 int32
 CmdPack(const FCmdPackOptions& Options)
 {
@@ -220,6 +276,38 @@ CmdPack(const FCmdPackOptions& Options)
 		return -1;
 	}
 
+	std::string P4HaveBuffer;
+
+	if (Options.bRunP4Have)
+	{
+		UNSYNC_LOG(L"Runinng `p4 have`");
+		int32 ReturnCode = RunSubprocess("p4 have ...", InputRoot, P4HaveBuffer);
+		if (ReturnCode != 0)
+		{
+			UNSYNC_ERROR("Error reported while running `p4 have`: %d", ReturnCode);
+			return -1;
+		}
+	}
+	else if (!Options.P4HavePath.empty())
+	{
+		UNSYNC_LOG(L"Loading p4 manifest file '%ls'", Options.P4HavePath.wstring().c_str());
+
+		FNativeFile P4HaveFile(Options.P4HavePath, EFileMode::ReadOnly);
+		if (!P4HaveFile.IsValid())
+		{
+			UNSYNC_ERROR(L"Could not open p4 manifest file '%ls'", Options.P4HavePath.wstring().c_str());
+			return -1;
+		}
+
+		P4HaveBuffer.resize(P4HaveFile.GetSize());
+		uint64 ReadBytes = P4HaveFile.Read(P4HaveBuffer.data(), 0, P4HaveFile.GetSize());
+		if (ReadBytes != P4HaveFile.GetSize())
+		{
+			UNSYNC_ERROR(L"Could not read the entire p4 manifest from '%ls'", Options.P4HavePath.wstring().c_str());
+			return -1;
+		}
+	}
+
 	// TODO: allow explicit output path
 	{
 		if (!EnsureDirectoryExists(ManifestRoot))
@@ -265,43 +353,63 @@ CmdPack(const FCmdPackOptions& Options)
 
 	const FPath DirectoryManifestPath = ManifestRoot / "manifest.bin";
 
-	std::string P4HaveBuffer;
-
 	FDirectoryManifest DirectoryManifest;
+	FDirectoryManifest OldDirectoryManifest;
 
-	if (!Options.P4HavePath.empty())
+	if (PathExists(DirectoryManifestPath))
 	{
-		UNSYNC_LOG(L"Loading p4 manifest file '%ls'", Options.P4HavePath.wstring().c_str());
+		UNSYNC_LOG(L"Loading previous manifest ");
+		UNSYNC_LOG_INDENT;
 
-		FNativeFile P4HaveFile(Options.P4HavePath, EFileMode::ReadOnly);
-		if (!P4HaveFile.IsValid())
+		if (LoadDirectoryManifest(OldDirectoryManifest, InputRoot, DirectoryManifestPath))
 		{
-			UNSYNC_ERROR(L"Could not open p4 manifest file '%ls'", Options.P4HavePath.wstring().c_str());
-			return -1;
+			if (OldDirectoryManifest.bHasFileRevisionControl)
+			{
+				UNSYNC_LOG(L"Loaded existing manifest with revision control data");
+			}
+			else
+			{
+				UNSYNC_LOG(L"Loaded existing manifest without revision control data");
+			}
 		}
+	}
 
-		P4HaveBuffer.resize(P4HaveFile.GetSize());
-		uint64 ReadBytes = P4HaveFile.Read(P4HaveBuffer.data(), 0, P4HaveFile.GetSize());
-		if (ReadBytes != P4HaveFile.GetSize())
-		{
-			UNSYNC_ERROR(L"Could not read the entire p4 manifest from '%ls'", Options.P4HavePath.wstring().c_str());
-			return -1;
-		}
+	if (!P4HaveBuffer.empty())
+	{
+		UNSYNC_LOG(L"Processing revision control data");
 
 		DirectoryManifest.Algorithm = Options.Algorithm;
 
 		BuildP4HaveSet(InputRoot, P4HaveBuffer, DirectoryManifest.Files);
+		DirectoryManifest.bHasFileRevisionControl = true;
 
 		UNSYNC_LOG(L"Loaded entries from p4 manifest: %llu", llu(DirectoryManifest.Files.size()));
 
 		UNSYNC_LOG(L"Updating file attributes");
-		auto UpdateFileMetadata = [](std::pair<const std::wstring, FFileManifest>& It)
+		auto UpdateFileMetadata = [&OldDirectoryManifest](std::pair<const std::wstring, FFileManifest>& It)
 		{
-			FFileAttributes Attrib = GetFileAttrib(It.second.CurrentPath);
+			if (OldDirectoryManifest.bHasFileRevisionControl)
+			{
+				auto OldFileIt = OldDirectoryManifest.Files.find(It.first);
+				if (OldFileIt != OldDirectoryManifest.Files.end())
+				{
+					if (It.second.RevisionControlIdentity == OldFileIt->second.RevisionControlIdentity)
+					{
+						It.second.Mtime		= OldFileIt->second.Mtime;
+						It.second.Size		= OldFileIt->second.Size;
+						It.second.bReadOnly = OldFileIt->second.bReadOnly;
+					}
+				}
+			}
 
-			It.second.Mtime		= Attrib.Mtime;
-			It.second.Size		= Attrib.Size;
-			It.second.bReadOnly = true;	 // treat all p4 files as read-only in the manifest
+			if (!It.second.IsValid())
+			{
+				FFileAttributes Attrib = GetFileAttrib(It.second.CurrentPath);
+
+				It.second.Mtime		= Attrib.Mtime;
+				It.second.Size		= Attrib.Size;
+				It.second.bReadOnly = true;	 // treat all p4 files as read-only in the manifest
+			}
 		};
 		ParallelForEach(DirectoryManifest.Files, UpdateFileMetadata);
 	}
@@ -360,30 +468,24 @@ CmdPack(const FCmdPackOptions& Options)
 	};
 
 	FComputeBlocksParams BlockParams;
-	BlockParams.Algorithm		 = Options.Algorithm;
-	BlockParams.BlockSize		 = Options.BlockSize;
+	BlockParams.Algorithm = Options.Algorithm;
+	BlockParams.BlockSize = Options.BlockSize;
 
 	// TODO: threading makes pack files non-deterministic...
 	// Perhaps a strictly ordered parallel pipeline mechanism could be implemented?
 	BlockParams.bAllowThreading	 = true;
 	BlockParams.OnBlockGenerated = OnBlockGenerated;
 
+	if (OldDirectoryManifest.IsValid())
 	{
-		UNSYNC_LOG(L"Loading previous manifest ");
-		FDirectoryManifest OldManifest;
-		if (PathExists(DirectoryManifestPath) && LoadDirectoryManifest(OldManifest, InputRoot, DirectoryManifestPath))
+		// Copy file blocks from old manifest, if possible
+		if (AlgorithmOptionsCompatible(DirectoryManifest.Algorithm, OldDirectoryManifest.Algorithm))
 		{
-			UNSYNC_LOG(L"Previous manifest loaded");
-
-			// Copy file blocks from old manifest, if possible
-			if (AlgorithmOptionsCompatible(DirectoryManifest.Algorithm, OldManifest.Algorithm))
-			{
-				MoveCompatibleManifestBlocks(DirectoryManifest, std::move(OldManifest));
-			}
-			else
-			{
-				UNSYNC_LOG(L"Incremental file block generation is not possible due to algorithm options mismatch");
-			}
+			MoveCompatibleManifestBlocks(DirectoryManifest, std::move(OldDirectoryManifest));
+		}
+		else
+		{
+			UNSYNC_LOG(L"Incremental file block generation is not possible due to algorithm options mismatch");
 		}
 	}
 
diff --git a/Engine/Source/Programs/Unsync/Private/UnsyncCmdPack.h b/Engine/Source/Programs/Unsync/Private/UnsyncCmdPack.h
--- a/Engine/Source/Programs/Unsync/Private/UnsyncCmdPack.h
+++ b/Engine/Source/Programs/Unsync/Private/UnsyncCmdPack.h
@@ -11,6 +11,7 @@ struct FCmdPackOptions
 {
 	FPath			  RootPath;
 	FPath			  P4HavePath;  // optional
+	bool			  bRunP4Have = false; // mutually exclusive with P4HavePath
 	FPath			  StorePath;   // optional
 	uint32			  BlockSize = uint32(64_KB);
 	FAlgorithmOptions Algorithm;
diff --git a/Engine/Source/Programs/Unsync/Private/UnsyncMain.cpp b/Engine/Source/Programs/Unsync/Private/UnsyncMain.cpp
--- a/Engine/Source/Programs/Unsync/Private/UnsyncMain.cpp
+++ b/Engine/Source/Programs/Unsync/Private/UnsyncMain.cpp
@@ -79,6 +79,7 @@ InnerMain(int Argc, char** Argv)
 	std::string				 ScavengeRootUtf8;
 	std::string				 P4HavePathUtf8;
 	std::string				 StorePathUtf8;
+	bool					 bRunP4Have			 = false;
 	bool					 bForceOperation	 = false;
 	bool					 bAllowInsecureTls	 = false;
 	bool					 bUseTls			 = false;
@@ -155,13 +156,27 @@ InnerMain(int Argc, char** Argv)
 
 	// Configure pack
 
-	CLI::App* SubPack = Cli.add_subcommand("pack", "EXPERIMENTAL: Generate manifest for a directory and store all referenced data in a compressed pack file");
-	SubPack->add_option("Input", InputFilenameUtf8, "Input directory path")->required();
-	SubPack->add_option("--p4havefile", P4HavePathUtf8, "Use `p4 have` output from a given file to explicitly specify files included in the manifest");
-	SubPack->add_option("--store",
-						StorePathUtf8,
-						"Use this location to store pack data (default: <Input>/.unsync/pack)");
-	SubCommands.push_back(SubPack);
+	CLI::App* SubPack = nullptr;
+	{
+		SubPack =
+			Cli.add_subcommand("pack",
+							   "EXPERIMENTAL: Generate manifest for a directory and store all referenced data in a compressed pack file");
+
+		SubPack->add_option("Input", InputFilenameUtf8, "Input directory path")->required();
+
+		auto P4HaveFileOpt =
+			SubPack->add_option("--p4havefile",
+								P4HavePathUtf8,
+								"Use `p4 have` output from a given file to explicitly specify files included in the manifest");
+
+		auto RunP4HaveOpt = SubPack->add_flag("--p4have", bRunP4Have, "Run `p4 have` when generating the dirctory pack");
+
+		SubPack->add_option("--store", StorePathUtf8, "Use this location to store pack data (default: <Input>/.unsync/pack)");
+
+		RunP4HaveOpt->excludes(P4HaveFileOpt);
+
+		SubCommands.push_back(SubPack);
+	}
 
 	// Configure push
 
@@ -573,12 +588,12 @@ InnerMain(int Argc, char** Argv)
 
 	if (Cli.got_subcommand(SubHash) || Cli.got_subcommand(SubPack))
 	{
-		UNSYNC_LOG(L"Using block size: %d KB", HashOrSyncBlockSize / 1024);
+		UNSYNC_VERBOSE(L"Using block size: %d KB", HashOrSyncBlockSize / 1024);
 	}
 
 	if (Cli.got_subcommand(SubDiff))
 	{
-		UNSYNC_LOG(L"Using block size: %d KB", DiffBlockSize / 1024);
+		UNSYNC_VERBOSE(L"Using block size: %d KB", DiffBlockSize / 1024);
 	}
 
 	if (Cli.got_subcommand(SubDiff))
@@ -770,6 +785,7 @@ InnerMain(int Argc, char** Argv)
 		PackOptions.RootPath	 = InputFilename;
 		PackOptions.P4HavePath	 = NormalizeFilenameUtf8(P4HavePathUtf8);
 		PackOptions.StorePath	 = NormalizeFilenameUtf8(StorePathUtf8);
+		PackOptions.bRunP4Have	 = bRunP4Have;
 		PackOptions.BlockSize	 = HashOrSyncBlockSize;
 		PackOptions.Algorithm	 = Algorithm;
 
