unsync - Basic implementation of unpack command

* Only filesystem based packs at this point

diff --git a/Engine/Source/Programs/Unsync/Private/UnsyncBuffer.h b/Engine/Source/Programs/Unsync/Private/UnsyncBuffer.h
--- a/Engine/Source/Programs/Unsync/Private/UnsyncBuffer.h
+++ b/Engine/Source/Programs/Unsync/Private/UnsyncBuffer.h
@@ -91,6 +91,7 @@ public:
 	}
 
 	operator FBufferView() const { return View(); }
+	operator FMutBufferView() { return MutView(); }
 
 private:
 	void EnsureCapacity(uint64 RequiredCapacity);
@@ -100,6 +101,14 @@ private:
 	uint64 BufferCapacity = 0;
 };
 
+template<typename T>
+TArrayView<T>
+ReinterpretView(FBufferView BufferView)
+{
+	const uint64 Count = BufferView.Size / sizeof(T);
+	return MakeView(reinterpret_cast<const T*>(BufferView.Data), Count);
+}
+
 struct FBufferPool
 {
 	UNSYNC_DISALLOW_COPY_ASSIGN(FBufferPool);
diff --git a/Engine/Source/Programs/Unsync/Private/UnsyncCmdPack.cpp b/Engine/Source/Programs/Unsync/Private/UnsyncCmdPack.cpp
--- a/Engine/Source/Programs/Unsync/Private/UnsyncCmdPack.cpp
+++ b/Engine/Source/Programs/Unsync/Private/UnsyncCmdPack.cpp
@@ -6,6 +6,7 @@
 #include "UnsyncHashTable.h"
 #include "UnsyncSerialization.h"
 #include "UnsyncThread.h"
+#include "UnsyncError.h"
 
 #include <stdio.h>
 #include <stdlib.h>
@@ -13,6 +14,8 @@
 
 namespace unsync {
 
+static constexpr uint64 GMaxPackFileSize = 1_GB;
+
 template<typename CallbackT>
 static void
 ForLines(std::string_view String, CallbackT Callback)
@@ -82,7 +85,7 @@ BuildP4HaveSet(const FPath& Root, std::string_view P4HaveDataUtf8, FDirectoryMan
 	ForLines(P4HaveDataUtf8, Callback);
 }
 
-struct FPackIndexEntry
+struct FPackIndexEntry	// structure is serialized
 {
 	FHash128 BlockHash		= {};
 	FHash128 CompressedHash = {};
@@ -91,6 +94,81 @@ struct FPackIndexEntry
 };
 static_assert(sizeof(FPackIndexEntry) == 40);
 
+struct FPackDatabase
+{
+	struct FEntry
+	{
+		FPackIndexEntry IndexEntry = {};
+		uint32			PackIndex  = ~0u;
+	};
+
+	// non-thread-safe
+	void Load(const FPath& PackRoot)
+	{
+		// TODO: cache index files locally if they're remote
+
+		const FPath ExpectedExtension = FPath(".unsync_index");
+		for (const std::filesystem::directory_entry& Dir : RecursiveDirectoryScan(PackRoot))
+		{
+			if (!Dir.is_regular_file())
+			{
+				continue;
+			}
+
+			const FPath&	IndexFilePath = Dir.path();
+			FPathStringView IndexFilePathView(IndexFilePath.native());
+			if (!IndexFilePathView.ends_with(ExpectedExtension.native()))
+			{
+				continue;
+			}
+
+			FPath PackFilePath = FPath(IndexFilePathView).replace_extension(".unsync_pack");
+			FFileAttributes PackAttrib = GetFileAttrib(PackFilePath);
+			if (!PackAttrib.bValid)
+			{
+				// TODO: also check file size, etc.
+				continue;
+			}
+
+			uint32 PackIndex = CheckedNarrow(PackFilenames.size());
+			PackFilenames.push_back(PackFilePath);
+			PackFileCache.push_back(nullptr);
+
+			FBuffer IndexEntries = ReadFileToBuffer(IndexFilePath);	 // TODO: add header & hash at the end
+
+			for (const FPackIndexEntry& IndexEntry : ReinterpretView<FPackIndexEntry>(IndexEntries))
+			{
+				UNSYNC_ASSERT(IndexEntry.Offset + IndexEntry.CompressedSize < GMaxPackFileSize);
+
+				FEntry DatabaseEntry;
+				DatabaseEntry.IndexEntry	   = IndexEntry;
+				DatabaseEntry.PackIndex		   = PackIndex;
+				BlockMap[IndexEntry.BlockHash] = DatabaseEntry;
+			}
+		}
+	}
+
+	// thread-safe
+	std::shared_ptr<FNativeFile> GetPackFile(uint32 PackFileIndex) const
+	{
+		std::lock_guard<std::mutex> LockGuard(FileCacheMutex);
+
+		if (!PackFileCache[PackFileIndex])
+		{
+			FNativeFile* NewFile		 = new FNativeFile(PackFilenames[PackFileIndex], EFileMode::ReadOnlyUnbuffered);
+			PackFileCache[PackFileIndex] = std::shared_ptr<FNativeFile>(NewFile);
+		}
+
+		return PackFileCache[PackFileIndex];
+	}
+
+	std::vector<FPath>						  PackFilenames;
+	THashMap<FHash128, FEntry>				  BlockMap;
+
+	mutable std::vector<std::shared_ptr<FNativeFile>> PackFileCache;
+	mutable std::mutex FileCacheMutex;
+};
+
 inline void
 AddHash(uint64* Accumulator, const FHash128& Hash)
 {
@@ -110,8 +188,6 @@ MakeHashFromParts(uint64* Parts)
 
 struct FPackWriteContext
 {
-	static constexpr uint64 MaxPackFileSize = 1_GB;
-
 	FPackWriteContext(const FPath& InOutputRoot) : OutputRoot(InOutputRoot) { Reset(); }
 	~FPackWriteContext() { FinishPack(); }
 
@@ -119,9 +195,9 @@ struct FPackWriteContext
 	{
 		std::lock_guard<std::mutex> LockGuard(Mutex);
 
-		UNSYNC_ASSERT(CompressedData.Size <= MaxPackFileSize);
+		UNSYNC_ASSERT(CompressedData.Size <= GMaxPackFileSize);
 
-		if (PackBuffer.Size() + CompressedData.Size > MaxPackFileSize)
+		if (PackBuffer.Size() + CompressedData.Size > GMaxPackFileSize)
 		{
 			FinishPack();
 		}
@@ -135,6 +211,8 @@ struct FPackWriteContext
 		IndexEntries.push_back(IndexEntry);
 		PackBuffer.Append(CompressedData);
 
+		UNSYNC_ASSERT(PackBuffer.Size() == IndexEntry.Offset + IndexEntry.CompressedSize);
+
 		AddHash(IndexFileHashSum, IndexEntry.BlockHash);
 	}
 
@@ -171,7 +249,7 @@ struct FPackWriteContext
 private:
 	void Reset()
 	{
-		PackBuffer.Reserve(MaxPackFileSize);
+		PackBuffer.Reserve(GMaxPackFileSize);
 		PackBuffer.Clear();
 		IndexEntries.clear();
 
@@ -323,7 +401,8 @@ CmdPack(const FCmdPackOptions& Options)
 		}
 	}
 
-	THashSet<FHash128> SeenBlockHashSet;
+	THashSet<FHash128> SeenBlockHashSet;  // TODO: use FPackDatabase, perhaps in some osrt of lightweight mode
+
 	UNSYNC_LOG(L"Loading block database");
 	{
 		UNSYNC_LOG_INDENT;
@@ -490,7 +569,35 @@ CmdPack(const FCmdPackOptions& Options)
 	}
 
 	UNSYNC_LOG(L"Computing file blocks");
-	UpdateDirectoryManifestBlocks(DirectoryManifest, InputRoot, BlockParams);
+	{
+		// Invalidate files with blocks that are not available in the pack DB
+		auto BlockValidator = [&SeenBlockHashSet](std::pair<const std::wstring, FFileManifest>& It)
+		{
+			FFileManifest& FileManifest = It.second;
+			uint64		   ValidFileSize = 0;
+
+			for (const FGenericBlock& Block : FileManifest.Blocks)
+			{
+				if (SeenBlockHashSet.find(Block.HashStrong.ToHash128()) == SeenBlockHashSet.end())
+				{
+					break;
+				}
+
+				ValidFileSize += Block.Size;
+			}
+
+			if (ValidFileSize != FileManifest.Size)
+			{
+				FileManifest.BlockSize = 0;
+				FileManifest.Blocks.clear();
+				FileManifest.MacroBlocks.clear();
+			}
+		};
+		ParallelForEach(DirectoryManifest.Files, BlockValidator);
+
+		// Scan files and generate new unique blocks
+		UpdateDirectoryManifestBlocks(DirectoryManifest, InputRoot, BlockParams);
+	}
 
 	uint64 ManifestUniqueBytes	   = 0;
 	uint64 ManifestCompressedBytes = 0;
@@ -581,4 +688,262 @@ CmdPack(const FCmdPackOptions& Options)
 	return 0;
 }
 
+bool
+BuildTargetFromPack(FIOWriter& Output, const FPackDatabase& PackDb, TArrayView<FGenericBlock> Manifest)
+{
+	struct FScheduleItem
+	{
+		uint32				   PackIndex  = ~0u;
+		const FPackIndexEntry* IndexEntry = nullptr;
+		const FGenericBlock*   Block;
+
+		bool operator<(const FScheduleItem& Other) const
+		{
+			if (PackIndex != Other.PackIndex)
+			{
+				return PackIndex < Other.PackIndex;
+			}
+			return IndexEntry->Offset < Other.IndexEntry->Offset;
+		}
+	};
+
+	std::vector<FScheduleItem> Schedule;
+	Schedule.reserve(Manifest.Size());
+
+	for (const FGenericBlock& Block : Manifest)
+	{
+		FHash128 BlockHash = Block.HashStrong.ToHash128();
+
+		auto FindIt = PackDb.BlockMap.find(BlockHash);
+		if (FindIt == PackDb.BlockMap.end())
+		{
+			UNSYNC_ERROR(L"Pack database does not contain required block");
+			return false;
+		}
+
+		const FPackDatabase::FEntry& PackEntry = FindIt->second;
+
+		FScheduleItem ScheduleItem;
+		ScheduleItem.PackIndex	= PackEntry.PackIndex;
+		ScheduleItem.IndexEntry = &PackEntry.IndexEntry;
+		ScheduleItem.Block		= &Block;
+
+		Schedule.push_back(ScheduleItem);
+	}
+
+	std::sort(Schedule.begin(), Schedule.end());
+
+	std::shared_ptr<FNativeFile> PackFile;
+	uint32						 CurrentPackId = ~0u;
+
+	// TODO: multi-threaded decompression
+
+	auto ReadCallback = [&Schedule, &Output](FIOBuffer Buffer, uint64 SourceOffset, uint64 ReadSize, uint64 ScheduleIndex)
+	{
+		const FScheduleItem& Item = Schedule[ScheduleIndex];
+
+		UNSYNC_ASSERT(ReadSize == Item.IndexEntry->CompressedSize);
+
+		FHash128 CompressedHash = HashBlake3Bytes<FHash128>(Buffer.GetData(), Buffer.GetSize());
+		UNSYNC_ASSERT(CompressedHash == Item.IndexEntry->CompressedHash);
+
+		FIOBuffer Decompressed = FIOBuffer::Alloc(Item.Block->Size, L"PackBlockDecompress");
+		if (!Decompress(Buffer.GetBufferView(), Decompressed.GetMutBufferView()))
+		{
+			UNSYNC_FATAL(L"Failed to decompress block while reading from pack");
+			return;
+		}
+
+		UNSYNC_ASSERT(Decompressed.GetSize() == Item.Block->Size);
+
+		Output.Write(Decompressed.GetData(), Item.Block->Offset, Item.Block->Size);
+	};
+
+	for (uint64 ScheduleIndex = 0; ScheduleIndex < Schedule.size(); ++ScheduleIndex)
+	{
+		const FScheduleItem& Item = Schedule[ScheduleIndex];
+		if (CurrentPackId != Item.PackIndex)
+		{
+			if (PackFile)
+			{
+				PackFile->FlushAll();
+			}
+
+			PackFile = PackDb.GetPackFile(Item.PackIndex);
+
+			if (!PackFile->IsValid())
+			{
+				UNSYNC_ERROR(L"Failed to open pack file '%ls'", PackDb.PackFilenames[Item.PackIndex].wstring().c_str());
+				return false;
+			}
+
+			CurrentPackId = Item.PackIndex;
+		}
+
+		UNSYNC_ASSERT(Item.IndexEntry->Offset < PackFile->GetSize());
+		UNSYNC_ASSERT(Item.IndexEntry->Offset + Item.IndexEntry->CompressedSize <= PackFile->GetSize());
+		PackFile->ReadAsync(Item.IndexEntry->Offset, Item.IndexEntry->CompressedSize, ScheduleIndex, ReadCallback);
+	}
+
+	if (PackFile)
+	{
+		PackFile->FlushAll();
+	}
+
+	return true;
+}
+
+struct FDirectoryCreationCache
+{
+	bool EnsureDirectoryExists(const FPath& Path)
+	{
+		std::lock_guard<std::mutex> LockGuard(Mutex);
+		if (CreatedDirectories.find(Path) != CreatedDirectories.end())
+		{
+			return true;
+		}
+		else
+		{
+			bool bCreated = unsync::EnsureDirectoryExists(Path);
+			if (bCreated)
+			{
+				CreatedDirectories.insert(Path);
+			}
+			return bCreated;
+		}
+	}
+
+	THashSet<FPath> CreatedDirectories;
+	std::mutex		Mutex;
+};
+
+bool
+SyncDirectoryFromPack(const FPath& OutputRoot, const FPackDatabase& PackDb, const FDirectoryManifest& NewDirectoryManifest)
+{
+	// TODO: multithreading
+	// TODO: incremental sync
+	// TODO: fetch blocks from server
+
+	FAtomicError Error;
+
+	FDirectoryCreationCache DirCache;
+
+	auto BuildTargetCallback = [&Error, &PackDb, &DirCache](const std::pair<const std::wstring, FFileManifest>& FileIt)
+	{
+		if (Error)
+		{
+			return;
+		}
+
+		const FFileManifest& FileManifest	= FileIt.second;
+		const FPath&		 TargetFilePath = FileManifest.CurrentPath;
+
+		FPath TargetFileParent = TargetFilePath.parent_path();
+		if (!DirCache.EnsureDirectoryExists(TargetFileParent))
+		{
+			UNSYNC_ERROR(L"Failed to create target directory '%ls'", TargetFileParent.wstring().c_str());
+			Error.Set(AppError("Failed to create target directory"));
+			return;
+		}
+
+		FNativeFile TargetFile(TargetFilePath, EFileMode::CreateWriteOnly, FileManifest.Size);
+		if (!TargetFile.IsValid())
+		{
+			UNSYNC_ERROR(L"Failed to create target file '%ls'", TargetFilePath.wstring().c_str());
+			Error.Set(AppError("Failed to create target file"));
+			return;
+		}
+
+		if (!BuildTargetFromPack(TargetFile, PackDb, MakeView(FileManifest.Blocks)))
+		{
+			UNSYNC_FATAL(L"Failed to reconstruct target file from pack '%ls'", TargetFilePath.wstring().c_str());
+			Error.Set(AppError("Failed to reconstruct target file from pack"));
+			return;
+		}
+	};
+
+	ParallelForEach(NewDirectoryManifest.Files, BuildTargetCallback);
+
+	return !Error;
+}
+
+int32
+CmdUnpack(const FCmdUnpackOptions& Options)
+{
+	UNSYNC_LOG(L"Unpacking snapshot '%hs' to '%ls'", Options.SnapshotName.c_str(), Options.OutputPath.wstring().c_str());
+	UNSYNC_LOG_INDENT;
+
+	if (!EnsureDirectoryExists(Options.OutputPath))
+	{
+		UNSYNC_ERROR(L"Failed to create output directory");
+		return -1;
+	}
+
+	const FPath& StoreRoot = Options.StorePath;
+	const FPath	 PackRoot  = StoreRoot / "pack";
+
+	FPackDatabase PackDb;
+	{
+		UNSYNC_LOG(L"Loading block database");
+		UNSYNC_LOG_INDENT;
+		PackDb.Load(PackRoot);
+	}
+
+	FPath	SnapshotPath   = Options.StorePath / (Options.SnapshotName + ".unsync_snapshot");
+	FBuffer SnapshotBuffer = ReadFileToBuffer(SnapshotPath);
+	if (SnapshotBuffer.Empty())
+	{
+		UNSYNC_ERROR(L"Failed read directory snapshot manifest");
+		return -1;
+	}
+
+	TArrayView<FGenericBlock> ManifestBlocks = ReinterpretView<FGenericBlock>(SnapshotBuffer);
+
+	uint64 ManifestFileSize = 0;
+	for (const FGenericBlock& Block : ManifestBlocks)
+	{
+		if (Block.Offset != ManifestFileSize)
+		{
+			UNSYNC_FATAL(L"Unexpected block offset found in the snapshot file");
+			return -1;
+		}
+		ManifestFileSize += Block.Size;
+	}
+
+	FBuffer ManifestBuffer;
+	{
+		ManifestBuffer.Resize(ManifestFileSize);
+		FMemReaderWriter ManifestWriter(ManifestBuffer);
+		if (!BuildTargetFromPack(ManifestWriter, PackDb, ManifestBlocks))
+		{
+			UNSYNC_FATAL(L"Failed to reconstruct directory manfiest from snapshot");
+			return -1;
+		}
+	}
+	
+
+	FDirectoryManifest NewDirectoryManifest;
+	{
+		FMemReader		   ManifestMemReader(ManifestBuffer);
+		FIOReaderStream	   ManifestReaderStream(ManifestMemReader);
+		if (!LoadDirectoryManifest(NewDirectoryManifest, Options.OutputPath, ManifestReaderStream))
+		{
+			UNSYNC_FATAL(L"Failed to deserialize directory manfiest from snapshot");
+			return -1;
+		}
+	}
+
+	UNSYNC_LOG(L"Writing target files");
+	{
+		UNSYNC_LOG_INDENT;
+		if (!SyncDirectoryFromPack(Options.OutputPath, PackDb, NewDirectoryManifest))
+		{
+			UNSYNC_FATAL(L"Failed to sync directory pack");
+			return -1;
+		}
+	}
+
+	return 0;
+}
+
 }  // namespace unsync
diff --git a/Engine/Source/Programs/Unsync/Private/UnsyncCmdPack.h b/Engine/Source/Programs/Unsync/Private/UnsyncCmdPack.h
--- a/Engine/Source/Programs/Unsync/Private/UnsyncCmdPack.h
+++ b/Engine/Source/Programs/Unsync/Private/UnsyncCmdPack.h
@@ -19,4 +19,13 @@ struct FCmdPackOptions
 
 int32 CmdPack(const FCmdPackOptions& Options);
 
+struct FCmdUnpackOptions
+{
+	FPath			  OutputPath;
+	FPath			  StorePath;
+	std::string		  SnapshotName;
+};
+
+int32 CmdUnpack(const FCmdUnpackOptions& Options);
+
 }
diff --git a/Engine/Source/Programs/Unsync/Private/UnsyncCompression.cpp b/Engine/Source/Programs/Unsync/Private/UnsyncCompression.cpp
--- a/Engine/Source/Programs/Unsync/Private/UnsyncCompression.cpp
+++ b/Engine/Source/Programs/Unsync/Private/UnsyncCompression.cpp
@@ -83,4 +83,9 @@ Decompress(const FBuffer& Buffer)
 	return Decompress(Buffer.Data(), Buffer.Size());
 }
 
+bool Decompress(FBufferView Input, FMutBufferView Output)
+{
+	return Decompress(Input.Data, Input.Size, Output.Data, Output.Size);
+}
+
 }  // namespace unsync
diff --git a/Engine/Source/Programs/Unsync/Private/UnsyncCompression.h b/Engine/Source/Programs/Unsync/Private/UnsyncCompression.h
--- a/Engine/Source/Programs/Unsync/Private/UnsyncCompression.h
+++ b/Engine/Source/Programs/Unsync/Private/UnsyncCompression.h
@@ -18,5 +18,6 @@ uint64 CompressInto(FBufferView Input, FMutBufferView Output, int ZstdCompressio
 FBuffer Decompress(const uint8* Data, uint64 DataSize);
 FBuffer Decompress(const FBuffer& Buffer);
 bool	Decompress(const uint8* InputData, uint64 InputDataSize, uint8* OutputData, uint64 OutputDataSize);
+bool	Decompress(FBufferView Input, FMutBufferView Output);
 
 }  // namespace unsync
diff --git a/Engine/Source/Programs/Unsync/Private/UnsyncFile.cpp b/Engine/Source/Programs/Unsync/Private/UnsyncFile.cpp
--- a/Engine/Source/Programs/Unsync/Private/UnsyncFile.cpp
+++ b/Engine/Source/Programs/Unsync/Private/UnsyncFile.cpp
@@ -279,10 +279,18 @@ FWindowsFile::OpenFileHandle(EFileMode InMode)
 	}
 }
 
+bool
+FWindowsFile::IsValid()
+{
+	std::lock_guard<std::mutex> LockGuard(Mutex);
+	return FileHandle != INVALID_HANDLE_VALUE;
+}
+
 void
 FWindowsFile::Close()
 {
-	FlushAll();
+	InternalFlushAll();
+
 	if (FileHandle != INVALID_HANDLE_VALUE)
 	{
 		CloseHandle(FileHandle);
@@ -301,6 +309,8 @@ FWindowsFile::Close()
 uint32
 FWindowsFile::CompleteReadCommand(Command& Cmd)
 {
+	// Expects that Mutex is locked
+
 	UNSYNC_ASSERT(Cmd.bActive);
 	const uint32 MaxAttempts = 100;
 
@@ -394,11 +404,13 @@ FWindowsFile::Write(const void* Data, uint64 DestOffset, uint64 TotalSize)
 {
 	// TODO: !!!!! fire-and-forget asynchronous writes !!!!!
 
+	std::lock_guard<std::mutex> LockGuard(Mutex);
+
 	UNSYNC_ASSERT(IsWritable(Mode));
 
 	if (!IsWriteOnly(Mode))
 	{
-		FlushAll();	 // flush any outstanding read requests before writing
+		InternalFlushAll();	 // flush any outstanding read requests before writing
 	}
 
 	LARGE_INTEGER Pos;
@@ -464,6 +476,8 @@ FWindowsFile::Write(const void* Data, uint64 DestOffset, uint64 TotalSize)
 uint64
 FWindowsFile::Read(void* Dest, uint64 SourceOffset, uint64 ReadSize)
 {
+	std::lock_guard<std::mutex> LockGuard(Mutex);
+
 	UNSYNC_ASSERTF((Mode & EFileMode::Unbuffered) == 0, L"Unbuffered files only support ReadAsync");
 	UNSYNC_ASSERT(IsReadable(Mode));
 
@@ -511,6 +525,8 @@ FWindowsFile::Read(void* Dest, uint64 SourceOffset, uint64 ReadSize)
 bool
 FWindowsFile::ReadAsync(uint64 SourceOffset, uint64 Size, uint64 UserData, IOCallback Callback)
 {
+	std::lock_guard<std::mutex> LockGuard(Mutex);
+
 	UNSYNC_ASSERT(IsReadable(Mode));
 
 	uint32 CmdIdx = ~0u;
@@ -575,7 +591,7 @@ FWindowsFile::ReadAsync(uint64 SourceOffset, uint64 Size, uint64 UserData, IOCal
 }
 
 void
-FWindowsFile::FlushAll()
+FWindowsFile::InternalFlushAll()
 {
 	for (uint32 I = 0; I < NUM_QUEUES; ++I)
 	{
@@ -586,9 +602,19 @@ FWindowsFile::FlushAll()
 	}
 }
 
+void
+FWindowsFile::FlushAll()
+{
+	std::lock_guard<std::mutex> LockGuard(Mutex);
+
+	InternalFlushAll();
+}
+
 void
 FWindowsFile::FlushOne()
 {
+	std::lock_guard<std::mutex> LockGuard(Mutex);
+
 	uint32 NumValidEvents			  = 0;
 	HANDLE Events[NUM_QUEUES]		  = {};
 	uint32 CommandIndices[NUM_QUEUES] = {};
diff --git a/Engine/Source/Programs/Unsync/Private/UnsyncFile.h b/Engine/Source/Programs/Unsync/Private/UnsyncFile.h
--- a/Engine/Source/Programs/Unsync/Private/UnsyncFile.h
+++ b/Engine/Source/Programs/Unsync/Private/UnsyncFile.h
@@ -9,6 +9,7 @@
 #include <atomic>
 #include <functional>
 #include <memory>
+#include <mutex>
 
 namespace unsync {
 
@@ -168,7 +169,7 @@ struct FWindowsFile : FIOReaderWriter
 	virtual void   FlushAll() override;
 	virtual void   FlushOne() override;
 	virtual uint64 GetSize() override { return FileSize; }
-	virtual bool   IsValid() override { return FileHandle != INVALID_HANDLE_VALUE; }
+	virtual bool   IsValid() override;
 	virtual void   Close() override;
 	virtual int32  GetError() override { return LastError; }
 
@@ -202,11 +203,16 @@ struct FWindowsFile : FIOReaderWriter
 	static constexpr uint32 UNBUFFERED_READ_ALIGNMENT = 4096;
 
 private:
+
+	// All internal methods expect the Mutex to be locked
+	void   InternalFlushAll();
 	uint32 CompleteReadCommand(Command& Cmd);
 	bool   OpenFileHandle(EFileMode InMode);
 
 private:
 	EFileMode Mode;
+
+	std::mutex Mutex;
 };
 using FNativeFile = FWindowsFile;
 #endif	// UNSYNC_PLATFORM_WINDOWS
@@ -307,6 +313,7 @@ struct FMemReader : FIOReader
 struct FMemReaderWriter : FMemReader, FIOReaderWriter
 {
 	FMemReaderWriter(uint8* InData, uint64 InDataSize);
+	FMemReaderWriter(FMutBufferView Buffer) : FMemReaderWriter(Buffer.Data, Buffer.Size) {}
 
 	// IOBase
 	virtual void Close() override
diff --git a/Engine/Source/Programs/Unsync/Private/UnsyncMain.cpp b/Engine/Source/Programs/Unsync/Private/UnsyncMain.cpp
--- a/Engine/Source/Programs/Unsync/Private/UnsyncMain.cpp
+++ b/Engine/Source/Programs/Unsync/Private/UnsyncMain.cpp
@@ -79,6 +79,7 @@ InnerMain(int Argc, char** Argv)
 	std::string				 ScavengeRootUtf8;
 	std::string				 P4HavePathUtf8;
 	std::string				 StorePathUtf8;
+	std::string				 SnapshotNameUtf8;
 	bool					 bRunP4Have			 = false;
 	bool					 bForceOperation	 = false;
 	bool					 bAllowInsecureTls	 = false;
@@ -178,6 +179,20 @@ InnerMain(int Argc, char** Argv)
 		SubCommands.push_back(SubPack);
 	}
 
+	CLI::App* SubUnpack = nullptr;
+	{
+		SubUnpack =
+			Cli.add_subcommand("unpack",
+							   "EXPERIMENTAL: Sync directory based on package snapshot");
+
+		SubUnpack->add_option("Output", OutputFilenameUtf8, "Output directory path")->required();
+
+		SubUnpack->add_option("--store", StorePathUtf8, "Pack storage path")->required();
+		SubUnpack->add_option("--snapshot", SnapshotNameUtf8, "Directory snapshot ID")->required();
+
+		SubCommands.push_back(SubUnpack);
+	}
+
 	// Configure push
 
 	CLI::App* SubPush = Cli.add_subcommand("push", "Loads a manifest from a directory and uploads referenced blocks to the remote server");
@@ -602,7 +617,7 @@ InnerMain(int Argc, char** Argv)
 		DefaultChunkingAlgorithm = EChunkingAlgorithmID::FixedBlocks;
 	}
 
-	if (Cli.got_subcommand(SubHash) || Cli.got_subcommand(SubDiff) || Cli.got_subcommand(SubPack))
+	if (Cli.got_subcommand(SubHash) || Cli.got_subcommand(SubDiff))
 	{
 		if (!bIncrementalMode)
 		{
@@ -778,7 +793,7 @@ InnerMain(int Argc, char** Argv)
 
 		return CmdHash(HashOptions);
 	}
-	if (Cli.got_subcommand(SubPack))
+	else if (Cli.got_subcommand(SubPack))
 	{
 		FCmdPackOptions PackOptions;
 
@@ -791,6 +806,16 @@ InnerMain(int Argc, char** Argv)
 
 		return CmdPack(PackOptions);
 	}
+	else if (Cli.got_subcommand(SubUnpack))
+	{
+		FCmdUnpackOptions UnpackOptions;
+
+		UnpackOptions.OutputPath   = OutputFilename;
+		UnpackOptions.SnapshotName = SnapshotNameUtf8;
+		UnpackOptions.StorePath	   = NormalizeFilenameUtf8(StorePathUtf8);
+
+		return CmdUnpack(UnpackOptions);
+	}
 	else if (Cli.got_subcommand(SubDiff))
 	{
 		FCmdDiffOptions DiffOptions;
