unsync - Extend unpack command logging

diff --git a/Engine/Source/Programs/Unsync/Private/UnsyncCmdPack.cpp b/Engine/Source/Programs/Unsync/Private/UnsyncCmdPack.cpp
--- a/Engine/Source/Programs/Unsync/Private/UnsyncCmdPack.cpp
+++ b/Engine/Source/Programs/Unsync/Private/UnsyncCmdPack.cpp
@@ -889,6 +889,8 @@ CmdUnpack(const FCmdUnpackOptions& Options)
 		PackDb.Load(PackRoot);
 	}
 
+	UNSYNC_LOG(L"Reading snapshot");
+
 	FPath	SnapshotPath   = Options.StorePath / (Options.SnapshotName + ".unsync_snapshot");
 	FBuffer SnapshotBuffer = ReadFileToBuffer(SnapshotPath);
 	if (SnapshotBuffer.Empty())
@@ -910,6 +912,8 @@ CmdUnpack(const FCmdUnpackOptions& Options)
 		ManifestFileSize += Block.Size;
 	}
 
+	UNSYNC_LOG(L"Reconstructing directory manifest");
+
 	FBuffer ManifestBuffer;
 	{
 		ManifestBuffer.Resize(ManifestFileSize);
@@ -924,6 +928,8 @@ CmdUnpack(const FCmdUnpackOptions& Options)
 
 	FDirectoryManifest NewDirectoryManifest;
 	{
+		UNSYNC_LOG(L"Loading directory manifest");
+
 		FMemReader		   ManifestMemReader(ManifestBuffer);
 		FIOReaderStream	   ManifestReaderStream(ManifestMemReader);
 		if (!LoadDirectoryManifest(NewDirectoryManifest, Options.OutputPath, ManifestReaderStream))
