unsync - Add dry run support to unpack command

* Useful for benchmarking and testing

diff --git a/Engine/Source/Programs/Unsync/Private/UnsyncCmdPack.cpp b/Engine/Source/Programs/Unsync/Private/UnsyncCmdPack.cpp
--- a/Engine/Source/Programs/Unsync/Private/UnsyncCmdPack.cpp
+++ b/Engine/Source/Programs/Unsync/Private/UnsyncCmdPack.cpp
@@ -797,6 +797,11 @@ struct FDirectoryCreationCache
 {
 	bool EnsureDirectoryExists(const FPath& Path)
 	{
+		if (GDryRun)
+		{
+			return true;
+		}
+
 		std::lock_guard<std::mutex> LockGuard(Mutex);
 		if (CreatedDirectories.find(Path) != CreatedDirectories.end())
 		{
@@ -820,7 +825,6 @@ struct FDirectoryCreationCache
 bool
 SyncDirectoryFromPack(const FPath& OutputRoot, const FPackDatabase& PackDb, const FDirectoryManifest& NewDirectoryManifest)
 {
-	// TODO: multithreading
 	// TODO: incremental sync
 	// TODO: fetch blocks from server
 
@@ -846,20 +850,45 @@ SyncDirectoryFromPack(const FPath& OutputRoot, const FPackDatabase& PackDb, cons
 			return;
 		}
 
-		FNativeFile TargetFile(TargetFilePath, EFileMode::CreateWriteOnly, FileManifest.Size);
-		if (!TargetFile.IsValid())
+		if (GDryRun)
 		{
-			UNSYNC_ERROR(L"Failed to create target file '%ls'", TargetFilePath.wstring().c_str());
-			Error.Set(AppError("Failed to create target file"));
-			return;
+			FNullReaderWriter NullFile(FileManifest.Size);
+			if (!BuildTargetFromPack(NullFile, PackDb, MakeView(FileManifest.Blocks)))
+			{
+				UNSYNC_FATAL(L"Failed to reconstruct target file from pack '%ls'", TargetFilePath.wstring().c_str());
+				Error.Set(AppError("Failed to reconstruct target file from pack"));
+				return;
+			}
 		}
-
-		if (!BuildTargetFromPack(TargetFile, PackDb, MakeView(FileManifest.Blocks)))
+		else
 		{
-			UNSYNC_FATAL(L"Failed to reconstruct target file from pack '%ls'", TargetFilePath.wstring().c_str());
-			Error.Set(AppError("Failed to reconstruct target file from pack"));
-			return;
+			// TODO: mark file writable if it exists
+			FNativeFile TargetFile(TargetFilePath, EFileMode::CreateWriteOnly, FileManifest.Size);
+
+			if (!TargetFile.IsValid())
+			{
+				UNSYNC_ERROR(L"Failed to create target file '%ls'", TargetFilePath.wstring().c_str());
+				Error.Set(AppError("Failed to create target file"));
+				return;
+			}
+
+			if (!BuildTargetFromPack(TargetFile, PackDb, MakeView(FileManifest.Blocks)))
+			{
+				UNSYNC_FATAL(L"Failed to reconstruct target file from pack '%ls'", TargetFilePath.wstring().c_str());
+				Error.Set(AppError("Failed to reconstruct target file from pack"));
+				return;
+			}
+
+			TargetFile.Close();
+
+			SetFileMtime(TargetFilePath, FileManifest.Mtime);
+
+			if (FileManifest.bReadOnly)
+			{
+				SetFileReadOnly(TargetFilePath, true);
+			}
 		}
+		
 	};
 
 	ParallelForEach(NewDirectoryManifest.Files, BuildTargetCallback);
@@ -949,6 +978,14 @@ CmdUnpack(const FCmdUnpackOptions& Options)
 		}
 	}
 
+	UNSYNC_LOG(L"Saving directory manifest");
+	const FPath ManifestRoot		  = Options.OutputPath / ".unsync";
+	const FPath DirectoryManifestPath = ManifestRoot / "manifest.bin";
+	if (!GDryRun && EnsureDirectoryExists(ManifestRoot))
+	{
+		SaveDirectoryManifest(NewDirectoryManifest, DirectoryManifestPath);
+	}
+
 	return 0;
 }
 
