unsync - Fix compile errors on Linux

* Avoid using std::filesystem::path in hash tables

diff --git a/Engine/Source/Programs/Unsync/Private/UnsyncCmdPack.cpp b/Engine/Source/Programs/Unsync/Private/UnsyncCmdPack.cpp
--- a/Engine/Source/Programs/Unsync/Private/UnsyncCmdPack.cpp
+++ b/Engine/Source/Programs/Unsync/Private/UnsyncCmdPack.cpp
@@ -803,7 +803,7 @@ struct FDirectoryCreationCache
 		}
 
 		std::lock_guard<std::mutex> LockGuard(Mutex);
-		if (CreatedDirectories.find(Path) != CreatedDirectories.end())
+		if (CreatedDirectories.find(Path.native()) != CreatedDirectories.end())
 		{
 			return true;
 		}
@@ -812,14 +812,14 @@ struct FDirectoryCreationCache
 			bool bCreated = unsync::EnsureDirectoryExists(Path);
 			if (bCreated)
 			{
-				CreatedDirectories.insert(Path);
+				CreatedDirectories.insert(Path.native());
 			}
 			return bCreated;
 		}
 	}
 
-	THashSet<FPath> CreatedDirectories;
-	std::mutex		Mutex;
+	THashSet<FPath::string_type> CreatedDirectories;
+	std::mutex					 Mutex;
 };
 
 bool
