unsync - Add basic implementation of a custom thread pool

* This will eventually replace ConcRT/PPL

diff --git a/Engine/Source/Programs/Unsync/Private/UnsyncThread.cpp b/Engine/Source/Programs/Unsync/Private/UnsyncThread.cpp
--- a/Engine/Source/Programs/Unsync/Private/UnsyncThread.cpp
+++ b/Engine/Source/Programs/Unsync/Private/UnsyncThread.cpp
@@ -8,6 +8,106 @@ namespace unsync {
 
 uint32 GMaxThreads = std::min<uint32>(UNSYNC_MAX_TOTAL_THREADS, std::thread::hardware_concurrency());
 
+FThreadPool GThreadPool;
+
+void
+FThreadPool::StartWorkers(uint32 NumWorkers)
+{
+	std::unique_lock<std::mutex> LockScope(Mutex);
+
+	while (Threads.size() < NumWorkers)
+	{
+		Threads.emplace_back(
+			[this]()
+			{
+				while (DoWork(true))
+				{
+				}
+			});
+	}
+}
+
+FThreadPool::~FThreadPool()
+{
+	bShutdown = true;
+	WorkerWakeupCondition.notify_all();
+
+	for (std::thread& Thread : Threads)
+	{
+		Thread.join();
+	}
+}
+
+FThreadPool::FTaskFunction
+FThreadPool::PopTask(bool bWaitForSignal)
+{
+	std::unique_lock<std::mutex> LockScope(Mutex);
+
+	if (bWaitForSignal)
+	{
+		auto WaitUntil = [this]() { return bShutdown || !Tasks.empty(); };
+		WorkerWakeupCondition.wait(LockScope, WaitUntil);
+	}
+
+	FThreadPool::FTaskFunction Result;
+
+	if (!Tasks.empty())
+	{
+		Result = std::move(Tasks.front());
+		Tasks.pop_front();
+	}
+
+	return Result;
+}
+
+uint64
+FThreadPool::PushTask(FTaskFunction&& Fun)
+{
+	if (Threads.empty())
+	{
+		Fun();
+		return NumTasksPushed;
+	}
+	else
+	{
+		std::unique_lock<std::mutex> LockScope(Mutex);
+		Tasks.push_back(std::forward<FTaskFunction>(Fun));
+		WorkerWakeupCondition.notify_one();
+		return ++NumTasksPushed;
+	}
+}
+
+bool
+FThreadPool::DoWork(bool bWaitForSignal)
+{
+	FTaskFunction Task = PopTask(bWaitForSignal);
+
+	if (Task)
+	{
+		Task();
+		++NumTasksCompleted;
+		return true;
+	}
+	else
+	{
+		return false;
+	}
+}
+
+void
+FThreadPool::WaitForFence(uint64 FenceValue)
+{
+	while (NumTasksCompleted < FenceValue)
+	{
+		if (DoWork(false))
+		{
+			continue;
+		}
+
+		std::this_thread::sleep_for(std::chrono::milliseconds(1));
+	}
+}
+
 #if UNSYNC_USE_CONCRT
 FConcurrencyPolicyScope::FConcurrencyPolicyScope(uint32 MaxConcurrency)
 {
diff --git a/Engine/Source/Programs/Unsync/Private/UnsyncThread.h b/Engine/Source/Programs/Unsync/Private/UnsyncThread.h
--- a/Engine/Source/Programs/Unsync/Private/UnsyncThread.h
+++ b/Engine/Source/Programs/Unsync/Private/UnsyncThread.h
@@ -7,6 +7,11 @@
 
 UNSYNC_THIRD_PARTY_INCLUDES_START
 #include <atomic>
+#include <condition_variable>
+#include <deque>
+#include <functional>
+#include <thread>
+#include <vector>
 #if UNSYNC_USE_CONCRT
 #	include <concrt.h>
 #	include <concurrent_queue.h>
@@ -16,7 +21,6 @@ UNSYNC_THIRD_PARTY_INCLUDES_START
 #	endif
 #else
 #	include <semaphore>
-#	include <thread>
 #	ifdef __APPLE__
 #		define UNSYNC_USE_MACH_SEMAPHORE 1
 #	endif
@@ -94,6 +98,45 @@ struct FThreadLogConfig
 void SchedulerSleep(uint32 Milliseconds);
 void SchedulerYield();
 
+class FThreadPool
+{
+public:
+
+	using FTaskFunction = std::function<void()>;
+
+	FThreadPool() = default;
+	~FThreadPool();
+
+	// Launches worker threads until total started worker count reaches NumWorkers.
+	// Does nothing if the number of already launched workers is lower than given value.
+	void StartWorkers(uint32 NumWorkers);
+
+	// Adds a task to the FIFO queue and returns its fence value
+	uint64 PushTask(FTaskFunction&& Fun);
+
+	// Block current thread until all tasks before and including given fence value
+	void WaitForFence(uint64 FenceValue);
+
+private:
+
+	// Try to execute a task and return whether there may be more tasks to run
+	bool DoWork(bool bWaitForSignal);
+
+	FTaskFunction PopTask(bool bWaitForSignal);
+
+	std::vector<std::thread>  Threads;
+	std::deque<FTaskFunction> Tasks;
+
+	std::atomic<uint64> NumTasksPushed;
+	std::atomic<uint64> NumTasksCompleted;
+
+	std::mutex				Mutex;
+	std::condition_variable WorkerWakeupCondition;
+	std::atomic<bool>		bShutdown;
+};
+
+extern FThreadPool GThreadPool;
+
 #if UNSYNC_USE_CONCRT
 
 // Cooperative semaphore implementation.
