unsync - Remove FThreadPool::WaitForFence, add TryExecuteTask

* Handling individual task completion should be implemented in higher-level primitives

diff --git a/Engine/Source/Programs/Unsync/Private/UnsyncThread.cpp b/Engine/Source/Programs/Unsync/Private/UnsyncThread.cpp
--- a/Engine/Source/Programs/Unsync/Private/UnsyncThread.cpp
+++ b/Engine/Source/Programs/Unsync/Private/UnsyncThread.cpp
@@ -20,7 +20,7 @@ FThreadPool::StartWorkers(uint32 NumWorkers)
 		Threads.emplace_back(
 			[this]()
 			{
-				while (DoWork(true))
+				while (DoWorkInternal(true))
 				{
 				}
 			});
@@ -60,32 +60,29 @@ FThreadPool::PopTask(bool bWaitForSignal)
 	return Result;
 }
 
-uint64
+void
 FThreadPool::PushTask(FTaskFunction&& Fun)
 {
 	if (Threads.empty())
 	{
 		Fun();
-		return NumTasksPushed;
 	}
 	else
 	{
 		std::unique_lock<std::mutex> LockScope(Mutex);
 		Tasks.push_back(std::forward<FTaskFunction>(Fun));
 		WorkerWakeupCondition.notify_one();
-		return ++NumTasksPushed;
 	}
 }
 
 bool
-FThreadPool::DoWork(bool bWaitForSignal)
+FThreadPool::DoWorkInternal(bool bWaitForSignal)
 {
 	FTaskFunction Task = PopTask(bWaitForSignal);
 
 	if (Task)
 	{
 		Task();
-		++NumTasksCompleted;
 		return true;
 	}
 	else
@@ -94,20 +91,6 @@ FThreadPool::DoWork(bool bWaitForSignal)
 	}
 }
 
-void
-FThreadPool::WaitForFence(uint64 FenceValue)
-{
-	while (NumTasksCompleted < FenceValue)
-	{
-		if (DoWork(false))
-		{
-			continue;
-		}
-
-		std::this_thread::sleep_for(std::chrono::milliseconds(1));
-	}
-}
-
 #if UNSYNC_USE_CONCRT
 FConcurrencyPolicyScope::FConcurrencyPolicyScope(uint32 MaxConcurrency)
 {
diff --git a/Engine/Source/Programs/Unsync/Private/UnsyncThread.h b/Engine/Source/Programs/Unsync/Private/UnsyncThread.h
--- a/Engine/Source/Programs/Unsync/Private/UnsyncThread.h
+++ b/Engine/Source/Programs/Unsync/Private/UnsyncThread.h
@@ -111,25 +111,23 @@ public:
 	// Does nothing if the number of already launched workers is lower than given value.
 	void StartWorkers(uint32 NumWorkers);
 
-	// Adds a task to the FIFO queue and returns its fence value
-	uint64 PushTask(FTaskFunction&& Fun);
+	// Adds a task to the FIFO queue
+	void PushTask(FTaskFunction&& Fun);
 
-	// Block current thread until all tasks before and including given fence value
-	void WaitForFence(uint64 FenceValue);
+	// Try to pop the next task from the queue and execute it on the current thread.
+	// Returns false if queue is empty, which may happen if worker threads have picked up the tasks already.
+	bool TryExecuteTask() { return DoWorkInternal(false); }
 
 private:
 
 	// Try to execute a task and return whether there may be more tasks to run
-	bool DoWork(bool bWaitForSignal);
+	bool DoWorkInternal(bool bWaitForSignal);
 
 	FTaskFunction PopTask(bool bWaitForSignal);
 
 	std::vector<std::thread>  Threads;
 	std::deque<FTaskFunction> Tasks;
 
-	std::atomic<uint64> NumTasksPushed;
-	std::atomic<uint64> NumTasksCompleted;
-
 	std::mutex				Mutex;
 	std::condition_variable WorkerWakeupCondition;
 	std::atomic<bool>		bShutdown;
