unsync - Add basic test for thread pool

diff --git a/Engine/Source/Programs/Unsync/Private/UnsyncTest.cpp b/Engine/Source/Programs/Unsync/Private/UnsyncTest.cpp
--- a/Engine/Source/Programs/Unsync/Private/UnsyncTest.cpp
+++ b/Engine/Source/Programs/Unsync/Private/UnsyncTest.cpp
@@ -834,6 +834,12 @@ RunTests(const std::string& Preset)
 		EStrongHashAlgorithmID::Blake3_256,
 	};
 
+	if (Preset == "thread" || Preset == "all")
+	{
+		extern void TestThread();
+		TestThread();
+	}
+
 	if (Preset == "misc" || Preset == "all")
 	{
 		TestMisc();
diff --git a/Engine/Source/Programs/Unsync/Private/UnsyncThread.cpp b/Engine/Source/Programs/Unsync/Private/UnsyncThread.cpp
--- a/Engine/Source/Programs/Unsync/Private/UnsyncThread.cpp
+++ b/Engine/Source/Programs/Unsync/Private/UnsyncThread.cpp
@@ -1,6 +1,7 @@
 // Copyright Epic Games, Inc. All Rights Reserved.
 
 #include "UnsyncThread.h"
+#include "UnsyncUtil.h"
 
 #include <chrono>
 
@@ -153,4 +154,43 @@ SchedulerYield()
 
 #endif	// UNSYNC_USE_CONCRT
 
+void
+TestThread()
+{
+	UNSYNC_LOG(L"TestThread()");
+	UNSYNC_LOG_INDENT;
+
+	{
+		UNSYNC_LOG(L"PushTask");
+
+		const uint32		NumTasks = 1000;
+		std::atomic<uint32> Counter	 = 0;
+
+		{
+			FThreadPool ThreadPool;
+			ThreadPool.StartWorkers(10);
+
+			uint32 RandomSeed = 1234;
+			for (uint32 i = 0; i < NumTasks; ++i)
+			{
+				uint32 R = Xorshift32(RandomSeed) % 10;
+				ThreadPool.PushTask(
+					[R, &Counter]
+					{
+						SchedulerSleep(1 + R);
+						Counter++;
+					});
+			}
+
+			while (ThreadPool.TryExecuteTask())
+			{
+			}
+
+			// thread pool destructor waits for outstanding tasks
+		}
+
+		UNSYNC_ASSERT(Counter == NumTasks)
+	}
+}
+
 }  // namespace unsync
diff --git a/Engine/Source/Programs/Unsync/Private/UnsyncThread.h b/Engine/Source/Programs/Unsync/Private/UnsyncThread.h
--- a/Engine/Source/Programs/Unsync/Private/UnsyncThread.h
+++ b/Engine/Source/Programs/Unsync/Private/UnsyncThread.h
@@ -2,8 +2,10 @@
 
 #pragma once
 
-#include "UnsyncUtil.h"
+#include "UnsyncCommon.h"
+
 #include "UnsyncLog.h"
+#include "UnsyncUtil.h"
 
 UNSYNC_THIRD_PARTY_INCLUDES_START
 #include <atomic>
