unsync - Fix TLS subject initialization in FHttpConnection

* Bump version to 1.0.62

diff --git a/Engine/Source/Programs/Unsync/Private/UnsyncCore.cpp b/Engine/Source/Programs/Unsync/Private/UnsyncCore.cpp
--- a/Engine/Source/Programs/Unsync/Private/UnsyncCore.cpp
+++ b/Engine/Source/Programs/Unsync/Private/UnsyncCore.cpp
@@ -27,7 +27,7 @@ UNSYNC_THIRD_PARTY_INCLUDES_START
 #include <md5-sse2.h>
 UNSYNC_THIRD_PARTY_INCLUDES_END
 
-#define UNSYNC_VERSION_STR "1.0.61-dev4"
+#define UNSYNC_VERSION_STR "1.0.61"
 
 namespace unsync {
 
diff --git a/Engine/Source/Programs/Unsync/Private/UnsyncHttp.cpp b/Engine/Source/Programs/Unsync/Private/UnsyncHttp.cpp
--- a/Engine/Source/Programs/Unsync/Private/UnsyncHttp.cpp
+++ b/Engine/Source/Programs/Unsync/Private/UnsyncHttp.cpp
@@ -491,9 +491,15 @@ FHttpConnection::FHttpConnection(const std::string_view InHostAddress, uint16 In
 {
 	if (InTlsSettings)
 	{
-		UNSYNC_ASSERT(!InTlsSettings->Subject.empty());
+		if (InTlsSettings->Subject.empty())
+		{
+			TlsSubject = std::string(InHostAddress);
+		}
+		else
+		{
+			TlsSubject = std::string(InTlsSettings->Subject);
+		}
 
-		TlsSubject			  = std::string(InTlsSettings->Subject);
 		bTlsVerifyCertificate = InTlsSettings->bVerifyCertificate;
 		bTlsVerifySubject	  = InTlsSettings->bVerifySubject;
 		if (InTlsSettings->CACert.Data)
