unsync - Move code related to building the target file into a separate .h/cpp

diff --git a/Engine/Source/Programs/Unsync/CMakeLists.txt b/Engine/Source/Programs/Unsync/CMakeLists.txt
--- a/Engine/Source/Programs/Unsync/CMakeLists.txt
+++ b/Engine/Source/Programs/Unsync/CMakeLists.txt
@@ -94,6 +94,8 @@ set(src
 	Private/UnsyncSerialization.h
 	Private/UnsyncSocket.cpp
 	Private/UnsyncSocket.h
+	Private/UnsyncTarget.cpp
+	Private/UnsyncTarget.h
 	Private/UnsyncTest.cpp
 	Private/UnsyncTest.h
 	Private/UnsyncThread.cpp
diff --git a/Engine/Source/Programs/Unsync/Private/UnsyncCmdDiff.cpp b/Engine/Source/Programs/Unsync/Private/UnsyncCmdDiff.cpp
--- a/Engine/Source/Programs/Unsync/Private/UnsyncCmdDiff.cpp
+++ b/Engine/Source/Programs/Unsync/Private/UnsyncCmdDiff.cpp
@@ -2,6 +2,7 @@
 
 #include "UnsyncCmdDiff.h"
 #include "UnsyncFile.h"
+#include "UnsyncTarget.h"
 
 namespace unsync {
 
diff --git a/Engine/Source/Programs/Unsync/Private/UnsyncCmdPatch.cpp b/Engine/Source/Programs/Unsync/Private/UnsyncCmdPatch.cpp
--- a/Engine/Source/Programs/Unsync/Private/UnsyncCmdPatch.cpp
+++ b/Engine/Source/Programs/Unsync/Private/UnsyncCmdPatch.cpp
@@ -2,6 +2,7 @@
 
 #include "UnsyncCmdPatch.h"
 #include "UnsyncFile.h"
+#include "UnsyncTarget.h"
 
 namespace unsync {
 
diff --git a/Engine/Source/Programs/Unsync/Private/UnsyncCore.cpp b/Engine/Source/Programs/Unsync/Private/UnsyncCore.cpp
--- a/Engine/Source/Programs/Unsync/Private/UnsyncCore.cpp
+++ b/Engine/Source/Programs/Unsync/Private/UnsyncCore.cpp
@@ -13,9 +13,9 @@
 #include "UnsyncSerialization.h"
 #include "UnsyncThread.h"
 #include "UnsyncUtil.h"
+#include "UnsyncTarget.h"
 
 #include <condition_variable>
-#include <deque>
 #include <filesystem>
 #include <mutex>
 #include <string>
@@ -33,8 +33,6 @@ namespace unsync {
 
 bool GDryRun = false;
 
-static const uint32 MAX_ACTIVE_READERS = 64;  // std::max(4, std::thread::hardware_concurrency();
-
 inline uint32
 ComputeMinVariableBlockSize(uint32 BlockSize)
 {
@@ -1057,25 +1055,6 @@ OptimizeNeedList(const std::vector<FNeedBlock>& Input, uint64 MaxMergedBlockSize
 	return Result;
 }
 
-FBuffer
-BuildTargetBuffer(const uint8*				SourceData,
-				  uint64					SourceSize,
-				  const uint8*				BaseData,
-				  uint64					BaseSize,
-				  const FNeedList&			NeedList,
-				  const FBuildTargetParams& Params)
-{
-	FMemReader SourceReader(SourceData, SourceSize);
-	FMemReader BaseReader(BaseData, BaseSize);
-	return BuildTargetBuffer(SourceReader, BaseReader, NeedList, Params);
-}
-
-struct FReadSchedule
-{
-	std::vector<FCopyCommand> Blocks;
-	std::deque<uint64>		  Requests;	 // unique block request indices sorted small to large
-};
-
 FReadSchedule
 BuildReadSchedule(const std::vector<FNeedBlock>& Blocks)
 {
@@ -1100,708 +1079,6 @@ BuildReadSchedule(const std::vector<FNeedBlock>& Blocks)
 	return Result;
 }
 
-class FBlockCache
-{
-public:
-	FBuffer							BlockData;
-	THashMap<FHash128, FBufferView> BlockMap;  // Decompressed block data by hash
-};
-
-FBuffer
-BuildTargetBuffer(FIOReader& SourceProvider, FIOReader& BaseProvider, const FNeedList& NeedList, const FBuildTargetParams& Params)
-{
-	FBuffer				Result;
-	const FNeedListSize SizeInfo = ComputeNeedListSize(NeedList);
-	Result.Resize(SizeInfo.TotalBytes);
-	FMemReaderWriter ResultWriter(Result.Data(), Result.Size());
-	BuildTarget(ResultWriter, SourceProvider, BaseProvider, NeedList, Params);
-	return Result;
-}
-
-void
-DownloadBlocks(FProxyPool&					  ProxyPool,
-			   const std::vector<FNeedBlock>& OriginalUniqueNeedBlocks,
-			   const FBlockDownloadCallback&  CompletionCallback)
-{
-	const uint32		ParentThreadIndent	 = GLogIndent;
-	const bool			bParentThreadVerbose = GLogVerbose;
-	std::atomic<uint64> NumActiveLogThreads	 = {};
-
-	THashSet<FHash128>			   DownloadedBlocks;
-	std::vector<FNeedBlock>		   RemainingNeedBlocks;
-	const std::vector<FNeedBlock>* NeedBlocks = &OriginalUniqueNeedBlocks;
-
-	std::atomic<bool> bGotError = false;
-
-	const uint32 MaxAttempts = 30;
-	for (uint32 Attempt = 0; Attempt <= MaxAttempts && !bGotError; ++Attempt)
-	{
-		if (!ProxyPool.IsValid())
-		{
-			UNSYNC_WARNING(L"FProxy connection is not valid.")
-		}
-
-		if (Attempt != 0)
-		{
-			uint64 NumRemainingBlocks = RemainingNeedBlocks.size();
-			UNSYNC_WARNING(L"Failed to download %lld blocks from proxy. Retry attempt %d of %d ...",
-						   NumRemainingBlocks,
-						   Attempt,
-						   MaxAttempts);
-
-			SchedulerSleep(1000);
-		}
-
-		struct FDownloadBatch
-		{
-			uint64 Begin;
-			uint64 End;
-			uint64 SizeBytes;
-		};
-		std::vector<FDownloadBatch> Batches;
-		Batches.push_back(FDownloadBatch{});
-
-		const uint64 MaxBytesPerBatch = ProxyPool.RemoteDesc.Protocol == EProtocolFlavor::Jupiter ? 16_MB : 128_MB;
-
-		for (uint64 I = 0; I < NeedBlocks->size() && !bGotError; ++I)
-		{
-			const FNeedBlock& Block = (*NeedBlocks)[I];
-
-			if (Batches.back().SizeBytes + Block.Size < MaxBytesPerBatch)
-			{
-				Batches.back().End = I + 1;
-			}
-			else
-			{
-				UNSYNC_ASSERT(Batches.back().End == I);
-				FDownloadBatch NewBatch = {};
-				NewBatch.Begin			= I;
-				NewBatch.End			= I + 1;
-				NewBatch.SizeBytes		= 0;
-				Batches.push_back(NewBatch);
-			}
-
-			Batches.back().SizeBytes += Block.Size;
-		}
-
-		UNSYNC_VERBOSE2(L"Download batches: %lld", Batches.size());
-
-		FTaskGroup DownloadTasks;
-		std::mutex DownloadedBlocksMutex;
-
-		for (FDownloadBatch Batch : Batches)
-		{
-			ProxyPool.ParallelDownloadSemaphore.Acquire();
-			DownloadTasks.run([NeedBlocks,
-							   Batch,
-							   ParentThreadIndent,
-							   bParentThreadVerbose,
-							   &DownloadedBlocksMutex,
-							   &ProxyPool,
-							   &DownloadedBlocks,
-							   &CompletionCallback,
-							   &NumActiveLogThreads,
-							   &bGotError] {
-				if (bGotError)
-				{
-					ProxyPool.ParallelDownloadSemaphore.Release();
-					return;
-				}
-
-				FThreadElectScope  AllowVerbose(NumActiveLogThreads, bParentThreadVerbose);
-				FLogVerbosityScope VerboseScope(AllowVerbose);
-				FLogIndentScope	   IndentScope(ParentThreadIndent, true);
-
-				auto Proxy = ProxyPool.Alloc();
-				if (Proxy)
-				{
-					const FNeedBlock*	   Blocks		  = NeedBlocks->data();
-					TArrayView<FNeedBlock> NeedBlocksView = {Blocks + Batch.Begin, Blocks + Batch.End};
-
-					FDownloadResult DownloadResult =
-						Proxy->Download(NeedBlocksView,
-										[&DownloadedBlocksMutex, &CompletionCallback, &DownloadedBlocks](const FDownloadedBlock& Block,
-																										 FHash128 BlockHash) {
-											{
-												std::lock_guard<std::mutex> LockGuard(DownloadedBlocksMutex);
-												DownloadedBlocks.insert(BlockHash);
-											}
-											CompletionCallback(Block, BlockHash);
-										});
-
-					if (DownloadResult.IsOk())
-					{
-						ProxyPool.Dealloc(std::move(Proxy));
-					}
-					else if (DownloadResult.GetError().RetryMode == EDownloadRetryMode::Abort)
-					{
-						bGotError = true;
-						ProxyPool.Invalidate();
-					}
-				}
-				ProxyPool.ParallelDownloadSemaphore.Release();
-			});
-		}
-
-		DownloadTasks.wait();
-
-		if (DownloadedBlocks.size() == OriginalUniqueNeedBlocks.size() || bGotError)
-		{
-			break;
-		}
-
-		RemainingNeedBlocks.clear();
-		for (const FNeedBlock& Block : OriginalUniqueNeedBlocks)
-		{
-			if (DownloadedBlocks.find(Block.Hash.ToHash128()) == DownloadedBlocks.end())  // #wip-widehash
-			{
-				RemainingNeedBlocks.push_back(Block);
-			}
-		}
-		NeedBlocks = &RemainingNeedBlocks;
-	}
-}
-
-FBuildTargetResult
-BuildTarget(FIOWriter& Output, FIOReader& Source, FIOReader& Base, const FNeedList& NeedList, const FBuildTargetParams& Params)
-{
-	UNSYNC_LOG_INDENT;
-
-	FBuildTargetResult BuildResult;
-
-	auto TimeBegin = TimePointNow();
-
-	const FNeedListSize			 SizeInfo		  = ComputeNeedListSize(NeedList);
-	const EStrongHashAlgorithmID StrongHasher	  = Params.StrongHasher;
-	FProxyPool*					 ProxyPool		  = Params.ProxyPool;
-	FScavengeDatabase*			 ScavengeDatabase = Params.ScavengeDatabase;
-	FBlockCache*				 BlockCache		  = Params.BlockCache;
-
-	if (SizeInfo.TotalBytes != Output.GetSize())
-	{
-		UNSYNC_ERROR(L"Output size is %llu, but expected to be %llu", llu(Output.GetSize()), llu(SizeInfo.TotalBytes));
-		return BuildResult;
-	}
-
-	UNSYNC_ASSERT(SizeInfo.TotalBytes == Output.GetSize());
-
-	FAtomicError Error;
-
-	std::atomic<bool> bBaseDataCopyTaskDone	  = false;
-	std::atomic<bool> bSourceDataCopyTaskDone = false;
-	std::atomic<bool> bWaitingForBaseData	  = false;
-
-	struct FStats
-	{
-		std::atomic<uint64> WrittenBytesFromSource = {};
-		std::atomic<uint64> WrittenBytesFromBase   = {};
-	};
-	FStats Stats;
-
-	FSemaphore WriteSemaphore(MAX_ACTIVE_READERS);	// throttle writing tasks to avoid memory bloat
-	FTaskGroup WriteTasks;
-
-	// Remember if parent thread has verbose logging and indentation
-	const bool	 bAllowVerboseLog = GLogVerbose;
-	const uint32 LogIndent		  = GLogIndent;
-
-	// Copy the need list as it will be progressively filtered by various sources:
-	// - block cache
-	// - scevenge from external local files
-	// - download from proxy
-	// - read from source file
-	// TODO: could use a bit array to indicate if the entry is valid instead of full copy
-	std::vector<FNeedBlock> FilteredSourceNeedList;
-
-	if (ScavengeDatabase)
-	{
-		THashSet<FHash128> ScavengedBlocks;
-
-		FScavengedBuildTargetResult ScavengeResult =
-			BuildTargetFromScavengedData(Output, NeedList.Source, *ScavengeDatabase, StrongHasher, ScavengedBlocks);
-
-		// Treat scavenged data the same as base for stats purposes
-		Stats.WrittenBytesFromBase += ScavengeResult.ScavengedBytes;
-
-		FilteredSourceNeedList = NeedList.Source;
-		auto FilterResult =
-			std::remove_if(FilteredSourceNeedList.begin(), FilteredSourceNeedList.end(), [&ScavengedBlocks](const FNeedBlock& Block) {
-				return ScavengedBlocks.find(Block.Hash.ToHash128()) != ScavengedBlocks.end();
-			});
-		FilteredSourceNeedList.erase(FilterResult, FilteredSourceNeedList.end());
-	}
-	else
-	{
-		FilteredSourceNeedList = NeedList.Source;
-	}
-
-	auto ProcessNeedList =
-		[bAllowVerboseLog, LogIndent, &Output, &Error, &WriteSemaphore, &WriteTasks, &bWaitingForBaseData, &Stats, &Params, SizeInfo](
-							   FIOReader&					  DataProvider,
-							   const std::vector<FNeedBlock>& NeedBlocks,
-							   uint64						  TotalCopySize,
-							   EBlockListType				  ListType,
-							   std::atomic<bool>&			  bCompletionFlag) {
-		FLogIndentScope IndentScope(LogIndent, true);
-		uint64			ReadBytesTotal = 0;
-
-		const wchar_t* ListName = ListType == EBlockListType::Source ? L"source" : L"base";
-
-		if (!DataProvider.IsValid())
-		{
-			UNSYNC_ERROR(L"Failed to read blocks from %ls. Stream is invalid.", ListName);
-			Error.Set(AppError(L"Failed to read blocks. Stream is invalid."));
-			return;
-		}
-
-		if (ListType == EBlockListType::Source && Params.SourceType == FBuildTargetParams::ESourceType::File &&
-			DataProvider.GetSize() != SizeInfo.TotalBytes)
-		{
-			UNSYNC_ERROR(L"File size is %llu, but expected to be %llu. File may have changed after manifest was generated.",
-						 llu(DataProvider.GetSize()),
-						 llu(SizeInfo.TotalBytes));
-			Error.Set(AppError(L"Failed to read source blocks. Size mismatch."));
-			return;
-		}
-
-		if (DataProvider.GetError())
-		{
-			UNSYNC_ERROR(L"Failed to read blocks from %ls. %hs", ListName, FormatSystemErrorMessage(DataProvider.GetError()).c_str());
-			Error.Set(AppError(L"Failed to read blocks"));
-			return;
-		}
-
-		UNSYNC_VERBOSE(L"Reading blocks from %ls", ListName);
-
-		FLogProgressScope ProgressLogger(TotalCopySize, ELogProgressUnits::MB);
-
-		FReadSchedule ReadSchedule = BuildReadSchedule(NeedBlocks);
-
-		UNSYNC_LOG_INDENT;
-		UNSYNC_VERBOSE(L"Using %d requests for %d total blocks (%.2f MB)",
-					   (int)ReadSchedule.Requests.size(),
-					   (int)NeedBlocks.size(),
-					   SizeMb(TotalCopySize));
-
-		static constexpr uint32 MaxActiveLargeRequests							= 2;
-		uint64					ActiveLargeRequestSizes[MaxActiveLargeRequests] = {};
-
-		while (!ReadSchedule.Requests.empty() && !Error)
-		{
-			uint64 BlockIndex = ~0ull;
-
-			uint32 LargeRequestSlot = ~0u;
-			for (uint32 I = 0; I < MaxActiveLargeRequests; ++I)
-			{
-				if (ActiveLargeRequestSizes[I] == 0)
-				{
-					LargeRequestSlot = I;
-					break;
-				}
-			}
-
-			if (LargeRequestSlot != ~0u)
-			{
-				BlockIndex = ReadSchedule.Requests.back();
-				ReadSchedule.Requests.pop_back();
-			}
-			else
-			{
-				BlockIndex = ReadSchedule.Requests.front();
-				ReadSchedule.Requests.pop_front();
-			}
-
-			const FCopyCommand& Block = ReadSchedule.Blocks[BlockIndex];
-
-			if (LargeRequestSlot != ~0u)
-			{
-				ActiveLargeRequestSizes[LargeRequestSlot] = Block.Size;
-			}
-
-			UNSYNC_ASSERTF(Block.SourceOffset + Block.Size <= DataProvider.GetSize(),
-						   L"Copy command is out of bounds. Offset %llu, size %llu ([%llu..%llu]), input size %llu.",
-						   llu(Block.SourceOffset),
-						   llu(Block.Size),
-						   llu(Block.SourceOffset),
-						   llu(Block.SourceOffset + Block.Size),
-						   DataProvider.GetSize());
-
-			uint64 ReadBytes = 0;
-
-			auto ReadCallback = [&ReadBytes, &Output, &WriteTasks, &Error, &WriteSemaphore, &Stats, Block, ListType](
-									FIOBuffer CmdBuffer,
-									uint64	  CmdOffset,
-									uint64	  CmdReadSize,
-									uint64	  CmdUserData) {
-				WriteSemaphore.Acquire();
-				WriteTasks.run([Buffer = MakeShared(std::move(CmdBuffer)),
-								CmdReadSize,
-								Block,
-								&Output, &Error,
-								&WriteSemaphore,
-								&Stats,
-								ListType]() {
-					const uint64 WrittenBytes = Output.Write(Buffer->GetData(), Block.TargetOffset, CmdReadSize);
-
-					if (ListType == EBlockListType::Source)
-					{
-						Stats.WrittenBytesFromSource += WrittenBytes;
-					}
-					else if (ListType == EBlockListType::Base)
-					{
-						Stats.WrittenBytesFromBase += WrittenBytes;
-					}
-					else
-					{
-						UNSYNC_FATAL(L"Unexpected block list type");
-					}
-
-					WriteSemaphore.Release();
-
-					if (WrittenBytes != CmdReadSize)
-					{
-						UNSYNC_FATAL(L"Expected to write %llu bytes, but written %llu", CmdReadSize, WrittenBytes);
-						Error.Set(AppError(L"Failed to write output"));
-					}
-				});
-
-				ReadBytes = CmdReadSize;
-
-				AddGlobalProgress(CmdReadSize, ListType);
-			};
-
-			DataProvider.ReadAsync(Block.SourceOffset, Block.Size, 0, ReadCallback);
-
-			for (uint32 I = 0; I < MaxActiveLargeRequests; ++I)
-			{
-				if (ActiveLargeRequestSizes[I] == ReadBytes)
-				{
-					ActiveLargeRequestSizes[I] = 0;
-					break;
-				}
-			}
-
-			ReadBytesTotal += ReadBytes;
-
-			FLogVerbosityScope VerbosityScope(GLogVerbose ||
-											  (bAllowVerboseLog && (ListType == EBlockListType::Base && bWaitingForBaseData)));
-
-			ProgressLogger.Add(ReadBytes);
-
-			SchedulerYield();
-		}
-
-		DataProvider.FlushAll();
-
-		Output.FlushAll();
-
-		ProgressLogger.Complete();
-
-		bCompletionFlag = true;
-	};
-
-	FTaskGroup BackgroundTasks;
-	BackgroundTasks.run([SizeInfo, ProcessNeedList, &NeedList, &Base, &bBaseDataCopyTaskDone]() {
-		// Silence log messages on background task.
-		// Logging will be enabled if we're waiting for base data.
-		FLogVerbosityScope VerbosityScope(false);
-
-		if (SizeInfo.BaseBytes)
-		{
-			ProcessNeedList(Base, NeedList.Base, SizeInfo.BaseBytes, EBlockListType::Base, bBaseDataCopyTaskDone);
-		}
-	});
-
-	THashSet<FHash128> CachedBlocks;
-
-	if (BlockCache)
-	{
-		UNSYNC_VERBOSE(L"Writing blocks from cache");
-		UNSYNC_LOG_INDENT;
-
-		for (const FNeedBlock NeedBlock : FilteredSourceNeedList)
-		{
-			auto It = BlockCache->BlockMap.find(NeedBlock.Hash.ToHash128());
-			if (It != BlockCache->BlockMap.end())
-			{
-				FBufferView BlockBuffer = It->second;
-
-				const uint64 WrittenBytes = Output.Write(BlockBuffer.Data, NeedBlock.TargetOffset, BlockBuffer.Size);
-
-				Stats.WrittenBytesFromSource += WrittenBytes;
-
-				AddGlobalProgress(NeedBlock.Size, EBlockListType::Source);
-			}
-		}
-
-		auto FilterResult =
-			std::remove_if(FilteredSourceNeedList.begin(), FilteredSourceNeedList.end(), [BlockCache](const FNeedBlock& Block) {
-				return BlockCache->BlockMap.find(Block.Hash.ToHash128()) != BlockCache->BlockMap.end();
-			});
-
-		FilteredSourceNeedList.erase(FilterResult, FilteredSourceNeedList.end());
-	}
-
-	if (FilteredSourceNeedList.size() == 0)
-	{
-		// No more data is needed from source file
-	}
-	else if (ProxyPool && ProxyPool->IsValid())
-	{
-		UNSYNC_VERBOSE(L"Downloading blocks");
-		UNSYNC_LOG_INDENT;
-
-		struct FBlockWriteCmd
-		{
-			uint64 TargetOffset = 0;
-			uint64 Size			= 0;
-		};
-		std::mutex								DownloadedBlocksMutex;
-		THashSet<FHash128>						DownloadedBlocks;
-		THashMap<FHash128, FBlockWriteCmd>		BlockWriteMap;
-		THashMap<FHash128, std::vector<uint64>> BlockScatterMap;
-		std::vector<FNeedBlock>					UniqueNeedList;
-
-		uint64 EstimatedDownloadSize = 0;
-		for (const FNeedBlock& Block : FilteredSourceNeedList)
-		{
-			FBlockWriteCmd Cmd;
-			Cmd.TargetOffset  = Block.TargetOffset;
-			Cmd.Size		  = Block.Size;
-			auto InsertResult = BlockWriteMap.insert(std::make_pair(Block.Hash.ToHash128(), Cmd));	// #wip-widehash
-
-			if (InsertResult.second)
-			{
-				UniqueNeedList.push_back(Block);
-				EstimatedDownloadSize += Block.Size;
-			}
-			else
-			{
-				BlockScatterMap[Block.Hash.ToHash128()].push_back(Cmd.TargetOffset);  // #wip-widehash
-			}
-		}
-
-		FTaskGroup DecompressTasks;
-
-		FLogProgressScope DownloadProgressLogger(EstimatedDownloadSize, ELogProgressUnits::MB);
-
-		// limit how many decompression tasks can be queued up to avoid memory bloat
-		const uint64 MaxConcurrentDecompressionTasks = 64;
-		FSemaphore	 DecompressionSemaphore(MaxConcurrentDecompressionTasks);
-
-		const bool			bParentThreadVerbose = GLogVerbose;
-		const uint32		ParentThreadIndent	 = GLogIndent;
-		std::atomic<uint64> NumActiveLogThreads	 = {};
-		std::atomic<uint64> NumHashMismatches	 = {};
-
-		DownloadBlocks(
-			*ProxyPool,
-			UniqueNeedList,
-			[&NumActiveLogThreads,
-			 &DownloadProgressLogger,
-			 &BlockWriteMap,
-			 &DecompressionSemaphore,
-			 &DecompressTasks,
-			 &Output,
-			 StrongHasher,
-			 &BlockScatterMap,
-			 &DownloadedBlocksMutex,
-			 &DownloadedBlocks,
-			 &Error,
-			 &NumHashMismatches,
-			 &ParentThreadIndent,
-			 &bParentThreadVerbose,
-			 &Stats](const FDownloadedBlock& Block, FHash128 BlockHash) {
-				FLogIndentScope IndentScope(ParentThreadIndent, true);
-
-				{
-					FThreadElectScope  AllowVerbose(NumActiveLogThreads, bParentThreadVerbose);
-					FLogVerbosityScope VerboseScope(AllowVerbose);
-					DownloadProgressLogger.Add(Block.DecompressedSize);
-				}
-
-				auto WriteIt = BlockWriteMap.find(BlockHash);
-				if (WriteIt != BlockWriteMap.end())
-				{
-					FBlockWriteCmd Cmd = WriteIt->second;
-
-					// TODO: avoid this copy by storing IOBuffer in DownloadedBlock
-					bool	  bCompressed	 = Block.IsCompressed();
-					uint64	  DownloadedSize = bCompressed ? Block.CompressedSize : Block.DecompressedSize;
-					FIOBuffer DownloadedData = FIOBuffer::Alloc(DownloadedSize, L"downloaded_data");
-					memcpy(DownloadedData.GetData(), Block.Data, DownloadedSize);
-
-					DecompressionSemaphore.Acquire();
-
-					DecompressTasks.run([&Output,
-										 BlockHashUnaligned = BlockHash,
-										 Cmd,
-										 DownloadedData	  = std::make_shared<FIOBuffer>(std::move(DownloadedData)),
-										 DecompressedSize = Block.DecompressedSize,
-										 bCompressed,
-										 ParentThreadIndent,
-										 StrongHasher,
-										 &BlockScatterMap,
-										 &DownloadedBlocksMutex,
-										 &DownloadedBlocks,
-										 &DecompressionSemaphore,
-										 &Error,
-										 &NumHashMismatches,
-										 &Stats]() {
-						FLogIndentScope IndentScope(ParentThreadIndent, true);
-
-						// captured objects don't respect large alignment requirements
-						FHash128 BlockHash = BlockHashUnaligned;
-
-						bool	  bOk = false;
-						FIOBuffer DecompressedData;
-						if (DecompressedSize == Cmd.Size)
-						{
-							if (bCompressed)
-							{
-								DecompressedData = FIOBuffer::Alloc(Cmd.Size, L"decompressed_data");
-
-								bOk =
-									Decompress(DownloadedData->GetData(), DownloadedData->GetSize(), DecompressedData.GetData(), Cmd.Size);
-							}
-							else
-							{
-								DecompressedData = std::move(*DownloadedData);
-
-								bOk = true;
-							}
-
-							if (bOk)
-							{
-								FHash128 ActualBlockHash = ComputeHash(DecompressedData.GetData(), Cmd.Size, StrongHasher).ToHash128();
-
-								bOk = ActualBlockHash == BlockHash;
-
-								if (bOk)
-								{
-									uint64 WrittenBytes = Output.Write(DecompressedData.GetData(), Cmd.TargetOffset, Cmd.Size);
-
-									Stats.WrittenBytesFromSource += WrittenBytes;
-
-									AddGlobalProgress(Cmd.Size, EBlockListType::Source);
-
-									if (WrittenBytes != Cmd.Size)
-									{
-										bOk = false;
-										UNSYNC_FATAL(L"Expected to write %llu bytes, but written %llu", Cmd.Size, WrittenBytes);
-										Error.Set(AppError(L"Failed to write output"));
-									}
-								}
-								else
-								{
-									NumHashMismatches++;
-								}
-							}
-							else
-							{
-								UNSYNC_FATAL(L"Failed to decompress downloaded block");
-								Error.Set(AppError(L"Failed to decompress downloaded block"));
-							}
-						}
-
-						if (bOk)
-						{
-							auto ScatterIt = BlockScatterMap.find(BlockHash);
-							if (ScatterIt != BlockScatterMap.end())
-							{
-								const std::vector<uint64>& ScatterList = ScatterIt->second;
-								for (uint64 ScatterOffset : ScatterList)
-								{
-									uint64 WrittenBytes = Output.Write(DecompressedData.GetData(), ScatterOffset, Cmd.Size);
-
-									Stats.WrittenBytesFromSource += WrittenBytes;
-
-									if (WrittenBytes != Cmd.Size)
-									{
-										UNSYNC_FATAL(L"Expected to write %llu bytes, but written %llu", Cmd.Size, WrittenBytes);
-										Error.Set(AppError(L"Failed to write output"));
-									}
-								}
-
-								AddGlobalProgress(Cmd.Size * ScatterList.size(), EBlockListType::Source);
-							}
-						}
-
-						if (bOk)
-						{
-							std::lock_guard<std::mutex> LockGuard(DownloadedBlocksMutex);
-							DownloadedBlocks.insert(BlockHash);
-						}
-
-						DecompressionSemaphore.Release();
-					});
-				}
-			});
-
-		DecompressTasks.wait();
-
-		DownloadProgressLogger.Complete();
-
-		if (NumHashMismatches.load() != 0)
-		{
-			UNSYNC_WARNING(L"Found block hash mismatches while downloading data from proxy. Mismatching blocks: %llu.",
-						   NumHashMismatches.load());
-		}
-
-		{
-			auto FilterResult =
-				std::remove_if(FilteredSourceNeedList.begin(), FilteredSourceNeedList.end(), [&DownloadedBlocks](const FNeedBlock& Block) {
-					return DownloadedBlocks.find(Block.Hash.ToHash128()) != DownloadedBlocks.end();
-				});
-			FilteredSourceNeedList.erase(FilterResult, FilteredSourceNeedList.end());
-		}
-
-		if (!FilteredSourceNeedList.empty())
-		{
-			uint64 RemainingBytes = ComputeSize(FilteredSourceNeedList);
-
-			UNSYNC_VERBOSE(L"Did not receive %d blocks from proxy. Must read %.2f MB from source.",
-						   FilteredSourceNeedList.size(),
-						   SizeMb(RemainingBytes));
-
-			UNSYNC_LOG_INDENT;
-
-			ProcessNeedList(Source, FilteredSourceNeedList, RemainingBytes, EBlockListType::Source, bSourceDataCopyTaskDone);
-		}
-		else
-		{
-			bSourceDataCopyTaskDone = true;
-		}
-	}
-	else if (SizeInfo.SourceBytes)
-	{
-		ProcessNeedList(Source, FilteredSourceNeedList, SizeInfo.SourceBytes, EBlockListType::Source, bSourceDataCopyTaskDone);
-	}
-
-	if (!bBaseDataCopyTaskDone)
-	{
-		UNSYNC_VERBOSE(L"Reading blocks from base");
-	}
-
-	bWaitingForBaseData = true;
-	BackgroundTasks.wait();
-	WriteTasks.wait();
-
-	double Duration = DurationSec(TimeBegin, TimePointNow());
-	UNSYNC_VERBOSE(L"Done in %.3f sec (%.3f MB / sec)", Duration, SizeMb(double(SizeInfo.TotalBytes) / Duration));
-
-	if (GLogVerbose)
-	{
-		LogGlobalProgress();
-	}
-
-	BuildResult.bSuccess	= !Error;
-	BuildResult.BaseBytes	= Stats.WrittenBytesFromBase;
-	BuildResult.SourceBytes = Stats.WrittenBytesFromSource;
-
-	return BuildResult;
-}
-
 bool
 IsSynchronized(const FNeedList& NeedList, const FGenericBlockArray& SourceBlocks)
 {
@@ -2584,7 +1861,7 @@ DownloadFileIfNewer(const FRemoteDesc& RemoteDesc, const FPath& Source, const FP
 	std::string SourceParentUtf8   = ConvertWideToUtf8(SourceParent.wstring());
 	std::string SourceFileNameUtf8 = ConvertWideToUtf8(SourceFileName.wstring());
 
-	// TODO: could have a dedicated single file stat query 
+	// TODO: could have a dedicated single file stat query
 	TResult<FDirectoryListing> DirectoryListingResult = ProxyQuery::ListDirectory(RemoteDesc, SourceParentUtf8);
 	if (DirectoryListingResult.IsError())
 	{
@@ -3790,26 +3067,6 @@ FSyncFilter::Resolve(const FPath& Filename) const
 	return Result;
 }
 
-struct FPatchHeader
-{
-	static constexpr uint64 VALIDATION_BLOCK_SIZE = 16_MB;
-
-	static constexpr uint64 MAGIC	= 0x3E63942C4C9ECE16ull;
-	static constexpr uint64 VERSION = 2;
-
-	uint64				   Magic					 = MAGIC;
-	uint64				   Version					 = VERSION;
-	uint64				   SourceSize				 = 0;
-	uint64				   BaseSize					 = 0;
-	uint64				   NumSourceValidationBlocks = 0;
-	uint64				   NumBaseValidationBlocks	 = 0;
-	uint64				   NumSourceBlocks			 = 0;
-	uint64				   NumBaseBlocks			 = 0;
-	uint64				   BlockSize				 = 0;
-	EWeakHashAlgorithmID   WeakHashAlgorithmId		 = EWeakHashAlgorithmID::Naive;
-	EStrongHashAlgorithmID StrongHashAlgorithmId	 = EStrongHashAlgorithmID::Blake3_128;
-};
-
 FBuffer
 GeneratePatch(const uint8*			 BaseData,
 			  uint64				 BaseDataSize,
@@ -3906,206 +3163,6 @@ GeneratePatch(const uint8*			 BaseData,
 	return Result;
 }
 
-template<typename BlockType>
-static bool
-BlocksEqual(const std::vector<BlockType>& A, std::vector<BlockType>& B)
-{
-	if (A.size() != B.size())
-	{
-		return false;
-	}
-
-	for (uint64 I = 0; I < A.size(); ++I)
-	{
-		if (A[I].HashStrong != B[I].HashStrong || A[I].HashWeak != B[I].HashWeak || A[I].Size != B[I].Size || A[I].Offset != B[I].Offset)
-		{
-			return false;
-		}
-	}
-
-	return true;
-}
-
-struct FMemStream
-{
-	FMemStream(const uint8* InData, uint64 InDataSize) : Data(InData), DataRw(nullptr), SIZE(InDataSize), Offset(0) {}
-
-	FMemStream(uint8* InData, uint64 InDataSize) : Data(InData), DataRw(InData), SIZE(InDataSize), Offset(0) {}
-
-	bool Read(void* Dest, uint64 ReadSize)
-	{
-		if (Offset + ReadSize <= SIZE)
-		{
-			uint8* DestBytes = reinterpret_cast<uint8*>(Dest);
-			memcpy(DestBytes, Data + Offset, ReadSize);
-			Offset += ReadSize;
-			return true;
-		}
-		else
-		{
-			return false;
-		}
-	}
-
-	bool Write(const void* Source, uint64 WriteSize)
-	{
-		if (Offset + WriteSize <= SIZE)
-		{
-			const uint8* SourceBytes = reinterpret_cast<const uint8*>(Source);
-			memcpy(DataRw + Offset, SourceBytes, WriteSize);
-			Offset += WriteSize;
-			return true;
-		}
-		else
-		{
-			return false;
-		}
-	}
-
-	const uint8* Data;
-	uint8*		 DataRw;
-	const uint64 SIZE;
-	uint64		 Offset;
-};
-
-FBuffer
-BuildTargetWithPatch(const uint8* PatchData, uint64 PatchSize, const uint8* BaseData, uint64 BaseDataSize)
-{
-	const EChunkingAlgorithmID ChunkingMode = EChunkingAlgorithmID::FixedBlocks;
-
-	FBuffer Result;
-
-	FBuffer	   DecompressedPatch = Decompress(PatchData, PatchSize);
-	FMemStream Stream(DecompressedPatch.Data(), DecompressedPatch.Size());
-
-	FPatchHeader Header;
-	Stream.Read(&Header, sizeof(Header));
-
-	if (Header.Magic != FPatchHeader::MAGIC)
-	{
-		UNSYNC_ERROR(L"Patch file header mismatch. Expected %llX, got %llX", (long long)FPatchHeader::MAGIC, (long long)Header.Magic);
-		return Result;
-	}
-
-	if (Header.Version != FPatchHeader::VERSION)
-	{
-		UNSYNC_ERROR(L"Patch file version mismatch. Expected %lld, got %lld", (long long)FPatchHeader::VERSION, (long long)Header.Version);
-		return Result;
-	}
-
-	FHash128 ExpectedHeaderHash = HashBlake3Bytes<FHash128>(Stream.Data, Stream.Offset);
-	FHash128 HeaderHash;
-	Stream.Read(&HeaderHash, sizeof(HeaderHash));
-	if (HeaderHash != ExpectedHeaderHash)
-	{
-		UNSYNC_ERROR(L"Patch header hash mismatch");
-		return Result;
-	}
-
-	FGenericBlockArray SourceValidationExpected(Header.NumSourceValidationBlocks);
-	FGenericBlockArray BaseValidationExpected(Header.NumBaseValidationBlocks);
-
-	Stream.Read(&SourceValidationExpected[0], SourceValidationExpected.size() * sizeof(SourceValidationExpected[0]));
-	Stream.Read(&BaseValidationExpected[0], BaseValidationExpected.size() * sizeof(BaseValidationExpected[0]));
-
-	FGenericBlockArray BaseValidation;
-
-	{
-		FLogVerbosityScope VerbosityScope(false);
-		FAlgorithmOptions  Algorithm;
-		Algorithm.ChunkingAlgorithmId	= ChunkingMode;
-		Algorithm.WeakHashAlgorithmId	= Header.WeakHashAlgorithmId;
-		Algorithm.StrongHashAlgorithmId = Header.StrongHashAlgorithmId;
-
-		BaseValidation = ComputeBlocks(BaseData, BaseDataSize, FPatchHeader::VALIDATION_BLOCK_SIZE, Algorithm);
-	}
-
-	if (!BlocksEqual(BaseValidation, BaseValidationExpected))
-	{
-		UNSYNC_ERROR(L"Base file mismatch");
-		return Result;
-	}
-
-	FPatchCommandList PatchCommands;
-	PatchCommands.Source.resize(Header.NumSourceBlocks);
-	PatchCommands.Base.resize(Header.NumBaseBlocks);
-
-	Stream.Read(&PatchCommands.Source[0], Header.NumSourceBlocks * sizeof(PatchCommands.Source[0]));
-	Stream.Read(&PatchCommands.Base[0], Header.NumBaseBlocks * sizeof(PatchCommands.Base[0]));
-
-	FHash128 ExpectedBlockHash = HashBlake3Bytes<FHash128>(Stream.Data, Stream.Offset);
-	FHash128 BlockHash;
-	Stream.Read(&BlockHash, sizeof(BlockHash));
-	if (BlockHash != ExpectedBlockHash)
-	{
-		UNSYNC_ERROR(L"Patch block hash mismatch");
-		return Result;
-	}
-
-	const uint8* SourceData		= Stream.Data + Stream.Offset;
-	uint64		 SourceDataSize = Stream.SIZE - Stream.Offset;
-
-	// TODO: generate the proper source need list from the start
-	uint64 SourceOffset = 0;
-	for (FCopyCommand& Block : PatchCommands.Source)
-	{
-		Block.SourceOffset = SourceOffset;
-		SourceOffset += Block.Size;
-	}
-
-	FNeedList NeedList;
-	NeedList.Base.reserve(PatchCommands.Base.size());
-	NeedList.Source.reserve(PatchCommands.Source.size());
-	NeedList.Sequence.reserve(PatchCommands.Sequence.size());
-
-	for (const FCopyCommand& Cmd : PatchCommands.Base)
-	{
-		FNeedBlock Block;
-		static_cast<FCopyCommand&>(Block) = Cmd;
-		NeedList.Base.push_back(Block);
-	}
-
-	for (const FCopyCommand& Cmd : PatchCommands.Source)
-	{
-		FNeedBlock Block;
-		static_cast<FCopyCommand&>(Block) = Cmd;
-		NeedList.Source.push_back(Block);
-	}
-
-	for (const FHash128& Hash : PatchCommands.Sequence)
-	{
-		NeedList.Sequence.push_back(Hash);
-	}
-
-	FBuildTargetParams BuildParams;
-	BuildParams.StrongHasher = Header.StrongHashAlgorithmId;
-	BuildParams.SourceType	 = FBuildTargetParams::ESourceType::Patch;
-
-	Result = BuildTargetBuffer(SourceData, SourceDataSize, BaseData, BaseDataSize, NeedList, BuildParams);
-
-	FGenericBlockArray SourceValidation;
-
-	{
-		FLogVerbosityScope VerbosityScope(false);
-
-		FAlgorithmOptions Algorithm;
-		Algorithm.ChunkingAlgorithmId	= ChunkingMode;
-		Algorithm.WeakHashAlgorithmId	= Header.WeakHashAlgorithmId;
-		Algorithm.StrongHashAlgorithmId = Header.StrongHashAlgorithmId;
-
-		SourceValidation = ComputeBlocks(Result.Data(), Result.Size(), FPatchHeader::VALIDATION_BLOCK_SIZE, Algorithm);
-	}
-
-	if (!BlocksEqual(SourceValidation, SourceValidationExpected))
-	{
-		UNSYNC_ERROR(L"Patched size file mismatch");
-		Result.Clear();
-		return Result;
-	}
-
-	return Result;
-}
-
 FNeedListSize
 ComputeNeedListSize(const FNeedList& NeedList)
 {
diff --git a/Engine/Source/Programs/Unsync/Private/UnsyncCore.h b/Engine/Source/Programs/Unsync/Private/UnsyncCore.h
--- a/Engine/Source/Programs/Unsync/Private/UnsyncCore.h
+++ b/Engine/Source/Programs/Unsync/Private/UnsyncCore.h
@@ -9,17 +9,21 @@
 #include "UnsyncSocket.h"
 #include "UnsyncUtil.h"
 #include "UnsyncManifest.h"
+#include "UnsyncHashTable.h"
 
 #include <stdint.h>
+#include <deque>
+#include <functional>
 #include <unordered_map>
 #include <unordered_set>
 #include <vector>
-#include <functional>
 
 namespace unsync {
 
 extern bool GDryRun;
 
+static constexpr uint32 MAX_ACTIVE_READERS = 64;
+
 class FProxy;
 class FProxyPool;
 class FBlockCache;
@@ -73,6 +77,21 @@ struct FNeedBlock : FCopyCommand
 	FGenericHash Hash = {};
 };
 
+struct FReadSchedule
+{
+	std::vector<FCopyCommand> Blocks;
+	std::deque<uint64>		  Requests;	 // unique block request indices sorted small to large
+};
+
+FReadSchedule BuildReadSchedule(const std::vector<FNeedBlock>& Blocks);
+
+class FBlockCache
+{
+public:
+	FBuffer							BlockData;
+	THashMap<FHash128, FBufferView> BlockMap;  // Decompressed block data by hash
+};
+
 inline uint64
 ComputeSize(const std::vector<FNeedBlock>& NeedBlocks)
 {
@@ -184,44 +203,6 @@ FNeedList DiffManifestBlocks(const FGenericBlockArray& SourceBlocks, const FGene
 
 std::vector<FCopyCommand> OptimizeNeedList(const std::vector<FNeedBlock>& Input, uint64 MaxMergedBlockSize = 8_MB);
 
-struct FBuildTargetResult
-{
-	bool   bSuccess	   = false;
-	uint64 SourceBytes = 0;
-	uint64 BaseBytes   = 0;
-};
-
-struct FBuildTargetParams
-{
-	EStrongHashAlgorithmID StrongHasher;
-	FProxyPool*			   ProxyPool		= nullptr;
-	FBlockCache*		   BlockCache		= nullptr;
-	FScavengeDatabase*	   ScavengeDatabase = nullptr;
-
-	enum class ESourceType {
-		File,
-		Patch
-	};
-
-	ESourceType SourceType = ESourceType::File;
-};
-
-FBuildTargetResult BuildTarget(FIOWriter&				 Result,
-							   FIOReader&				 Source,
-							   FIOReader&				 Base,
-							   const FNeedList&			 NeedList,
-							   const FBuildTargetParams& Params);
-
-FBuffer BuildTargetBuffer(FIOReader& SourceProvider, FIOReader& BaseProvider, const FNeedList& NeedList, const FBuildTargetParams& Params);
-
-FBuffer BuildTargetBuffer(const uint8*				SourceData,
-						  uint64					SourceSize,
-						  const uint8*				BaseData,
-						  uint64					BaseSize,
-						  const FNeedList&			NeedList,
-						  const FBuildTargetParams& Params);
-
-FBuffer BuildTargetWithPatch(const uint8* PatchData, uint64 PatchSize, const uint8* BaseData, uint64 BaseSize);
 
 FBuffer GeneratePatch(const uint8*			 BaseData,
 					  uint64				 BaseDataSize,
diff --git a/Engine/Source/Programs/Unsync/Private/UnsyncProtocol.h b/Engine/Source/Programs/Unsync/Private/UnsyncProtocol.h
--- a/Engine/Source/Programs/Unsync/Private/UnsyncProtocol.h
+++ b/Engine/Source/Programs/Unsync/Private/UnsyncProtocol.h
@@ -182,6 +182,26 @@ struct FBlockPacket
 	FBuffer	 CompressedData;
 };
 
+struct FPatchHeader
+{
+	static constexpr uint64 VALIDATION_BLOCK_SIZE = 16_MB;
+
+	static constexpr uint64 MAGIC	= 0x3E63942C4C9ECE16ull;
+	static constexpr uint64 VERSION = 2;
+
+	uint64				   Magic					 = MAGIC;
+	uint64				   Version					 = VERSION;
+	uint64				   SourceSize				 = 0;
+	uint64				   BaseSize					 = 0;
+	uint64				   NumSourceValidationBlocks = 0;
+	uint64				   NumBaseValidationBlocks	 = 0;
+	uint64				   NumSourceBlocks			 = 0;
+	uint64				   NumBaseBlocks			 = 0;
+	uint64				   BlockSize				 = 0;
+	EWeakHashAlgorithmID   WeakHashAlgorithmId		 = EWeakHashAlgorithmID::Naive;
+	EStrongHashAlgorithmID StrongHashAlgorithmId	 = EStrongHashAlgorithmID::Blake3_128;
+};
+
 // Protocol V2: support for up to 256bit hashes
 
 struct FBlockRequest256
diff --git a/Engine/Source/Programs/Unsync/Private/UnsyncProxy.cpp b/Engine/Source/Programs/Unsync/Private/UnsyncProxy.cpp
--- a/Engine/Source/Programs/Unsync/Private/UnsyncProxy.cpp
+++ b/Engine/Source/Programs/Unsync/Private/UnsyncProxy.cpp
@@ -568,6 +568,11 @@ ProxyQuery::ListDirectory(const FRemoteDesc& Remote, const std::string& Path)
 
 	FHttpResponse Response = HttpRequest(Connection, Request);
 
+	if (!Response.Success())
+	{
+		return HttpError(Response.Code);
+	}
+
 	Response.Buffer.PushBack(0);
 
 	return FDirectoryListing::FromJson((const char*)Response.Buffer.Data());
diff --git a/Engine/Source/Programs/Unsync/Private/UnsyncTarget.cpp b/Engine/Source/Programs/Unsync/Private/UnsyncTarget.cpp
new file mode 100644
--- /dev/null
+++ b/Engine/Source/Programs/Unsync/Private/UnsyncTarget.cpp
@@ -0,0 +1,930 @@
+// Copyright Epic Games, Inc. All Rights Reserved.
+
+#include "UnsyncTarget.h"
+#include "UnsyncCompression.h"
+#include "UnsyncCore.h"
+#include "UnsyncError.h"
+#include "UnsyncHashTable.h"
+#include "UnsyncProgress.h"
+#include "UnsyncProxy.h"
+#include "UnsyncScavenger.h"
+#include "UnsyncThread.h"
+
+namespace unsync {
+
+void
+DownloadBlocks(FProxyPool&					  ProxyPool,
+			   const std::vector<FNeedBlock>& OriginalUniqueNeedBlocks,
+			   const FBlockDownloadCallback&  CompletionCallback)
+{
+	const uint32		ParentThreadIndent	 = GLogIndent;
+	const bool			bParentThreadVerbose = GLogVerbose;
+	std::atomic<uint64> NumActiveLogThreads	 = {};
+
+	THashSet<FHash128>			   DownloadedBlocks;
+	std::vector<FNeedBlock>		   RemainingNeedBlocks;
+	const std::vector<FNeedBlock>* NeedBlocks = &OriginalUniqueNeedBlocks;
+
+	std::atomic<bool> bGotError = false;
+
+	const uint32 MaxAttempts = 30;
+	for (uint32 Attempt = 0; Attempt <= MaxAttempts && !bGotError; ++Attempt)
+	{
+		if (!ProxyPool.IsValid())
+		{
+			UNSYNC_WARNING(L"FProxy connection is not valid.")
+		}
+
+		if (Attempt != 0)
+		{
+			uint64 NumRemainingBlocks = RemainingNeedBlocks.size();
+			UNSYNC_WARNING(L"Failed to download %lld blocks from proxy. Retry attempt %d of %d ...",
+						   NumRemainingBlocks,
+						   Attempt,
+						   MaxAttempts);
+
+			SchedulerSleep(1000);
+		}
+
+		struct FDownloadBatch
+		{
+			uint64 Begin;
+			uint64 End;
+			uint64 SizeBytes;
+		};
+		std::vector<FDownloadBatch> Batches;
+		Batches.push_back(FDownloadBatch{});
+
+		const uint64 MaxBytesPerBatch = ProxyPool.RemoteDesc.Protocol == EProtocolFlavor::Jupiter ? 16_MB : 128_MB;
+
+		for (uint64 I = 0; I < NeedBlocks->size() && !bGotError; ++I)
+		{
+			const FNeedBlock& Block = (*NeedBlocks)[I];
+
+			if (Batches.back().SizeBytes + Block.Size < MaxBytesPerBatch)
+			{
+				Batches.back().End = I + 1;
+			}
+			else
+			{
+				UNSYNC_ASSERT(Batches.back().End == I);
+				FDownloadBatch NewBatch = {};
+				NewBatch.Begin			= I;
+				NewBatch.End			= I + 1;
+				NewBatch.SizeBytes		= 0;
+				Batches.push_back(NewBatch);
+			}
+
+			Batches.back().SizeBytes += Block.Size;
+		}
+
+		UNSYNC_VERBOSE2(L"Download batches: %lld", Batches.size());
+
+		FTaskGroup DownloadTasks;
+		std::mutex DownloadedBlocksMutex;
+
+		for (FDownloadBatch Batch : Batches)
+		{
+			ProxyPool.ParallelDownloadSemaphore.Acquire();
+			DownloadTasks.run(
+				[NeedBlocks,
+				 Batch,
+				 ParentThreadIndent,
+				 bParentThreadVerbose,
+				 &DownloadedBlocksMutex,
+				 &ProxyPool,
+				 &DownloadedBlocks,
+				 &CompletionCallback,
+				 &NumActiveLogThreads,
+				 &bGotError]
+				{
+					if (bGotError)
+					{
+						ProxyPool.ParallelDownloadSemaphore.Release();
+						return;
+					}
+
+					FThreadElectScope  AllowVerbose(NumActiveLogThreads, bParentThreadVerbose);
+					FLogVerbosityScope VerboseScope(AllowVerbose);
+					FLogIndentScope	   IndentScope(ParentThreadIndent, true);
+
+					auto Proxy = ProxyPool.Alloc();
+					if (Proxy)
+					{
+						const FNeedBlock*	   Blocks		  = NeedBlocks->data();
+						TArrayView<FNeedBlock> NeedBlocksView = {Blocks + Batch.Begin, Blocks + Batch.End};
+
+						FDownloadResult DownloadResult =
+							Proxy->Download(NeedBlocksView,
+											[&DownloadedBlocksMutex, &CompletionCallback, &DownloadedBlocks](const FDownloadedBlock& Block,
+																											 FHash128 BlockHash)
+											{
+												{
+													std::lock_guard<std::mutex> LockGuard(DownloadedBlocksMutex);
+													DownloadedBlocks.insert(BlockHash);
+												}
+												CompletionCallback(Block, BlockHash);
+											});
+
+						if (DownloadResult.IsOk())
+						{
+							ProxyPool.Dealloc(std::move(Proxy));
+						}
+						else if (DownloadResult.GetError().RetryMode == EDownloadRetryMode::Abort)
+						{
+							bGotError = true;
+							ProxyPool.Invalidate();
+						}
+					}
+					ProxyPool.ParallelDownloadSemaphore.Release();
+				});
+		}
+
+		DownloadTasks.wait();
+
+		if (DownloadedBlocks.size() == OriginalUniqueNeedBlocks.size() || bGotError)
+		{
+			break;
+		}
+
+		RemainingNeedBlocks.clear();
+		for (const FNeedBlock& Block : OriginalUniqueNeedBlocks)
+		{
+			if (DownloadedBlocks.find(Block.Hash.ToHash128()) == DownloadedBlocks.end())  // #wip-widehash
+			{
+				RemainingNeedBlocks.push_back(Block);
+			}
+		}
+		NeedBlocks = &RemainingNeedBlocks;
+	}
+}
+
+FBuildTargetResult
+BuildTarget(FIOWriter& Output, FIOReader& Source, FIOReader& Base, const FNeedList& NeedList, const FBuildTargetParams& Params)
+{
+	UNSYNC_LOG_INDENT;
+
+	FBuildTargetResult BuildResult;
+
+	auto TimeBegin = TimePointNow();
+
+	const FNeedListSize			 SizeInfo		  = ComputeNeedListSize(NeedList);
+	const EStrongHashAlgorithmID StrongHasher	  = Params.StrongHasher;
+	FProxyPool*					 ProxyPool		  = Params.ProxyPool;
+	FScavengeDatabase*			 ScavengeDatabase = Params.ScavengeDatabase;
+	FBlockCache*				 BlockCache		  = Params.BlockCache;
+
+	if (SizeInfo.TotalBytes != Output.GetSize())
+	{
+		UNSYNC_ERROR(L"Output size is %llu, but expected to be %llu", llu(Output.GetSize()), llu(SizeInfo.TotalBytes));
+		return BuildResult;
+	}
+
+	UNSYNC_ASSERT(SizeInfo.TotalBytes == Output.GetSize());
+
+	FAtomicError Error;
+
+	std::atomic<bool> bBaseDataCopyTaskDone	  = false;
+	std::atomic<bool> bSourceDataCopyTaskDone = false;
+	std::atomic<bool> bWaitingForBaseData	  = false;
+
+	struct FStats
+	{
+		std::atomic<uint64> WrittenBytesFromSource = {};
+		std::atomic<uint64> WrittenBytesFromBase   = {};
+	};
+	FStats Stats;
+
+	FSemaphore WriteSemaphore(MAX_ACTIVE_READERS);	// throttle writing tasks to avoid memory bloat
+	FTaskGroup WriteTasks;
+
+	// Remember if parent thread has verbose logging and indentation
+	const bool	 bAllowVerboseLog = GLogVerbose;
+	const uint32 LogIndent		  = GLogIndent;
+
+	// Copy the need list as it will be progressively filtered by various sources:
+	// - block cache
+	// - scevenge from external local files
+	// - download from proxy
+	// - read from source file
+	// TODO: could use a bit array to indicate if the entry is valid instead of full copy
+	std::vector<FNeedBlock> FilteredSourceNeedList;
+
+	if (ScavengeDatabase)
+	{
+		THashSet<FHash128> ScavengedBlocks;
+
+		FScavengedBuildTargetResult ScavengeResult =
+			BuildTargetFromScavengedData(Output, NeedList.Source, *ScavengeDatabase, StrongHasher, ScavengedBlocks);
+
+		// Treat scavenged data the same as base for stats purposes
+		Stats.WrittenBytesFromBase += ScavengeResult.ScavengedBytes;
+
+		FilteredSourceNeedList = NeedList.Source;
+		auto FilterResult	   = std::remove_if(FilteredSourceNeedList.begin(),
+											FilteredSourceNeedList.end(),
+											[&ScavengedBlocks](const FNeedBlock& Block)
+											{ return ScavengedBlocks.find(Block.Hash.ToHash128()) != ScavengedBlocks.end(); });
+		FilteredSourceNeedList.erase(FilterResult, FilteredSourceNeedList.end());
+	}
+	else
+	{
+		FilteredSourceNeedList = NeedList.Source;
+	}
+
+	auto ProcessNeedList =
+		[bAllowVerboseLog, LogIndent, &Output, &Error, &WriteSemaphore, &WriteTasks, &bWaitingForBaseData, &Stats, &Params, SizeInfo](
+			FIOReader&					   DataProvider,
+			const std::vector<FNeedBlock>& NeedBlocks,
+			uint64						   TotalCopySize,
+			EBlockListType				   ListType,
+			std::atomic<bool>&			   bCompletionFlag)
+	{
+		FLogIndentScope IndentScope(LogIndent, true);
+		uint64			ReadBytesTotal = 0;
+
+		const wchar_t* ListName = ListType == EBlockListType::Source ? L"source" : L"base";
+
+		if (!DataProvider.IsValid())
+		{
+			UNSYNC_ERROR(L"Failed to read blocks from %ls. Stream is invalid.", ListName);
+			Error.Set(AppError(L"Failed to read blocks. Stream is invalid."));
+			return;
+		}
+
+		if (ListType == EBlockListType::Source && Params.SourceType == FBuildTargetParams::ESourceType::File &&
+			DataProvider.GetSize() != SizeInfo.TotalBytes)
+		{
+			UNSYNC_ERROR(L"File size is %llu, but expected to be %llu. File may have changed after manifest was generated.",
+						 llu(DataProvider.GetSize()),
+						 llu(SizeInfo.TotalBytes));
+			Error.Set(AppError(L"Failed to read source blocks. Size mismatch."));
+			return;
+		}
+
+		if (DataProvider.GetError())
+		{
+			UNSYNC_ERROR(L"Failed to read blocks from %ls. %hs", ListName, FormatSystemErrorMessage(DataProvider.GetError()).c_str());
+			Error.Set(AppError(L"Failed to read blocks"));
+			return;
+		}
+
+		UNSYNC_VERBOSE(L"Reading blocks from %ls", ListName);
+
+		FLogProgressScope ProgressLogger(TotalCopySize, ELogProgressUnits::MB);
+
+		FReadSchedule ReadSchedule = BuildReadSchedule(NeedBlocks);
+
+		UNSYNC_LOG_INDENT;
+		UNSYNC_VERBOSE(L"Using %d requests for %d total blocks (%.2f MB)",
+					   (int)ReadSchedule.Requests.size(),
+					   (int)NeedBlocks.size(),
+					   SizeMb(TotalCopySize));
+
+		static constexpr uint32 MaxActiveLargeRequests							= 2;
+		uint64					ActiveLargeRequestSizes[MaxActiveLargeRequests] = {};
+
+		while (!ReadSchedule.Requests.empty() && !Error)
+		{
+			uint64 BlockIndex = ~0ull;
+
+			uint32 LargeRequestSlot = ~0u;
+			for (uint32 I = 0; I < MaxActiveLargeRequests; ++I)
+			{
+				if (ActiveLargeRequestSizes[I] == 0)
+				{
+					LargeRequestSlot = I;
+					break;
+				}
+			}
+
+			if (LargeRequestSlot != ~0u)
+			{
+				BlockIndex = ReadSchedule.Requests.back();
+				ReadSchedule.Requests.pop_back();
+			}
+			else
+			{
+				BlockIndex = ReadSchedule.Requests.front();
+				ReadSchedule.Requests.pop_front();
+			}
+
+			const FCopyCommand& Block = ReadSchedule.Blocks[BlockIndex];
+
+			if (LargeRequestSlot != ~0u)
+			{
+				ActiveLargeRequestSizes[LargeRequestSlot] = Block.Size;
+			}
+
+			UNSYNC_ASSERTF(Block.SourceOffset + Block.Size <= DataProvider.GetSize(),
+						   L"Copy command is out of bounds. Offset %llu, size %llu ([%llu..%llu]), input size %llu.",
+						   llu(Block.SourceOffset),
+						   llu(Block.Size),
+						   llu(Block.SourceOffset),
+						   llu(Block.SourceOffset + Block.Size),
+						   DataProvider.GetSize());
+
+			uint64 ReadBytes = 0;
+
+			auto ReadCallback = [&ReadBytes, &Output, &WriteTasks, &Error, &WriteSemaphore, &Stats, Block, ListType](FIOBuffer CmdBuffer,
+																													 uint64	   CmdOffset,
+																													 uint64	   CmdReadSize,
+																													 uint64	   CmdUserData)
+			{
+				WriteSemaphore.Acquire();
+				WriteTasks.run(
+					[Buffer = MakeShared(std::move(CmdBuffer)), CmdReadSize, Block, &Output, &Error, &WriteSemaphore, &Stats, ListType]()
+					{
+						const uint64 WrittenBytes = Output.Write(Buffer->GetData(), Block.TargetOffset, CmdReadSize);
+
+						if (ListType == EBlockListType::Source)
+						{
+							Stats.WrittenBytesFromSource += WrittenBytes;
+						}
+						else if (ListType == EBlockListType::Base)
+						{
+							Stats.WrittenBytesFromBase += WrittenBytes;
+						}
+						else
+						{
+							UNSYNC_FATAL(L"Unexpected block list type");
+						}
+
+						WriteSemaphore.Release();
+
+						if (WrittenBytes != CmdReadSize)
+						{
+							UNSYNC_FATAL(L"Expected to write %llu bytes, but written %llu", CmdReadSize, WrittenBytes);
+							Error.Set(AppError(L"Failed to write output"));
+						}
+					});
+
+				ReadBytes = CmdReadSize;
+
+				AddGlobalProgress(CmdReadSize, ListType);
+			};
+
+			DataProvider.ReadAsync(Block.SourceOffset, Block.Size, 0, ReadCallback);
+
+			for (uint32 I = 0; I < MaxActiveLargeRequests; ++I)
+			{
+				if (ActiveLargeRequestSizes[I] == ReadBytes)
+				{
+					ActiveLargeRequestSizes[I] = 0;
+					break;
+				}
+			}
+
+			ReadBytesTotal += ReadBytes;
+
+			FLogVerbosityScope VerbosityScope(GLogVerbose ||
+											  (bAllowVerboseLog && (ListType == EBlockListType::Base && bWaitingForBaseData)));
+
+			ProgressLogger.Add(ReadBytes);
+
+			SchedulerYield();
+		}
+
+		DataProvider.FlushAll();
+
+		Output.FlushAll();
+
+		ProgressLogger.Complete();
+
+		bCompletionFlag = true;
+	};
+
+	FTaskGroup BackgroundTasks;
+	BackgroundTasks.run(
+		[SizeInfo, ProcessNeedList, &NeedList, &Base, &bBaseDataCopyTaskDone]()
+		{
+			// Silence log messages on background task.
+			// Logging will be enabled if we're waiting for base data.
+			FLogVerbosityScope VerbosityScope(false);
+
+			if (SizeInfo.BaseBytes)
+			{
+				ProcessNeedList(Base, NeedList.Base, SizeInfo.BaseBytes, EBlockListType::Base, bBaseDataCopyTaskDone);
+			}
+		});
+
+	THashSet<FHash128> CachedBlocks;
+
+	if (BlockCache)
+	{
+		UNSYNC_VERBOSE(L"Writing blocks from cache");
+		UNSYNC_LOG_INDENT;
+
+		for (const FNeedBlock NeedBlock : FilteredSourceNeedList)
+		{
+			auto It = BlockCache->BlockMap.find(NeedBlock.Hash.ToHash128());
+			if (It != BlockCache->BlockMap.end())
+			{
+				FBufferView BlockBuffer = It->second;
+
+				const uint64 WrittenBytes = Output.Write(BlockBuffer.Data, NeedBlock.TargetOffset, BlockBuffer.Size);
+
+				Stats.WrittenBytesFromSource += WrittenBytes;
+
+				AddGlobalProgress(NeedBlock.Size, EBlockListType::Source);
+			}
+		}
+
+		auto FilterResult = std::remove_if(FilteredSourceNeedList.begin(),
+										   FilteredSourceNeedList.end(),
+										   [BlockCache](const FNeedBlock& Block)
+										   { return BlockCache->BlockMap.find(Block.Hash.ToHash128()) != BlockCache->BlockMap.end(); });
+
+		FilteredSourceNeedList.erase(FilterResult, FilteredSourceNeedList.end());
+	}
+
+	if (FilteredSourceNeedList.size() == 0)
+	{
+		// No more data is needed from source file
+	}
+	else if (ProxyPool && ProxyPool->IsValid())
+	{
+		UNSYNC_VERBOSE(L"Downloading blocks");
+		UNSYNC_LOG_INDENT;
+
+		struct FBlockWriteCmd
+		{
+			uint64 TargetOffset = 0;
+			uint64 Size			= 0;
+		};
+		std::mutex								DownloadedBlocksMutex;
+		THashSet<FHash128>						DownloadedBlocks;
+		THashMap<FHash128, FBlockWriteCmd>		BlockWriteMap;
+		THashMap<FHash128, std::vector<uint64>> BlockScatterMap;
+		std::vector<FNeedBlock>					UniqueNeedList;
+
+		uint64 EstimatedDownloadSize = 0;
+		for (const FNeedBlock& Block : FilteredSourceNeedList)
+		{
+			FBlockWriteCmd Cmd;
+			Cmd.TargetOffset  = Block.TargetOffset;
+			Cmd.Size		  = Block.Size;
+			auto InsertResult = BlockWriteMap.insert(std::make_pair(Block.Hash.ToHash128(), Cmd));	// #wip-widehash
+
+			if (InsertResult.second)
+			{
+				UniqueNeedList.push_back(Block);
+				EstimatedDownloadSize += Block.Size;
+			}
+			else
+			{
+				BlockScatterMap[Block.Hash.ToHash128()].push_back(Cmd.TargetOffset);  // #wip-widehash
+			}
+		}
+
+		FTaskGroup DecompressTasks;
+
+		FLogProgressScope DownloadProgressLogger(EstimatedDownloadSize, ELogProgressUnits::MB);
+
+		// limit how many decompression tasks can be queued up to avoid memory bloat
+		const uint64 MaxConcurrentDecompressionTasks = 64;
+		FSemaphore	 DecompressionSemaphore(MaxConcurrentDecompressionTasks);
+
+		const bool			bParentThreadVerbose = GLogVerbose;
+		const uint32		ParentThreadIndent	 = GLogIndent;
+		std::atomic<uint64> NumActiveLogThreads	 = {};
+		std::atomic<uint64> NumHashMismatches	 = {};
+
+		DownloadBlocks(
+			*ProxyPool,
+			UniqueNeedList,
+			[&NumActiveLogThreads,
+			 &DownloadProgressLogger,
+			 &BlockWriteMap,
+			 &DecompressionSemaphore,
+			 &DecompressTasks,
+			 &Output,
+			 StrongHasher,
+			 &BlockScatterMap,
+			 &DownloadedBlocksMutex,
+			 &DownloadedBlocks,
+			 &Error,
+			 &NumHashMismatches,
+			 &ParentThreadIndent,
+			 &bParentThreadVerbose,
+			 &Stats](const FDownloadedBlock& Block, FHash128 BlockHash)
+			{
+				FLogIndentScope IndentScope(ParentThreadIndent, true);
+
+				{
+					FThreadElectScope  AllowVerbose(NumActiveLogThreads, bParentThreadVerbose);
+					FLogVerbosityScope VerboseScope(AllowVerbose);
+					DownloadProgressLogger.Add(Block.DecompressedSize);
+				}
+
+				auto WriteIt = BlockWriteMap.find(BlockHash);
+				if (WriteIt != BlockWriteMap.end())
+				{
+					FBlockWriteCmd Cmd = WriteIt->second;
+
+					// TODO: avoid this copy by storing IOBuffer in DownloadedBlock
+					bool	  bCompressed	 = Block.IsCompressed();
+					uint64	  DownloadedSize = bCompressed ? Block.CompressedSize : Block.DecompressedSize;
+					FIOBuffer DownloadedData = FIOBuffer::Alloc(DownloadedSize, L"downloaded_data");
+					memcpy(DownloadedData.GetData(), Block.Data, DownloadedSize);
+
+					DecompressionSemaphore.Acquire();
+
+					DecompressTasks.run(
+						[&Output,
+						 BlockHashUnaligned = BlockHash,
+						 Cmd,
+						 DownloadedData	  = std::make_shared<FIOBuffer>(std::move(DownloadedData)),
+						 DecompressedSize = Block.DecompressedSize,
+						 bCompressed,
+						 ParentThreadIndent,
+						 StrongHasher,
+						 &BlockScatterMap,
+						 &DownloadedBlocksMutex,
+						 &DownloadedBlocks,
+						 &DecompressionSemaphore,
+						 &Error,
+						 &NumHashMismatches,
+						 &Stats]()
+						{
+							FLogIndentScope IndentScope(ParentThreadIndent, true);
+
+							// captured objects don't respect large alignment requirements
+							FHash128 BlockHash = BlockHashUnaligned;
+
+							bool	  bOk = false;
+							FIOBuffer DecompressedData;
+							if (DecompressedSize == Cmd.Size)
+							{
+								if (bCompressed)
+								{
+									DecompressedData = FIOBuffer::Alloc(Cmd.Size, L"decompressed_data");
+
+									bOk = Decompress(DownloadedData->GetData(),
+													 DownloadedData->GetSize(),
+													 DecompressedData.GetData(),
+													 Cmd.Size);
+								}
+								else
+								{
+									DecompressedData = std::move(*DownloadedData);
+
+									bOk = true;
+								}
+
+								if (bOk)
+								{
+									FHash128 ActualBlockHash = ComputeHash(DecompressedData.GetData(), Cmd.Size, StrongHasher).ToHash128();
+
+									bOk = ActualBlockHash == BlockHash;
+
+									if (bOk)
+									{
+										uint64 WrittenBytes = Output.Write(DecompressedData.GetData(), Cmd.TargetOffset, Cmd.Size);
+
+										Stats.WrittenBytesFromSource += WrittenBytes;
+
+										AddGlobalProgress(Cmd.Size, EBlockListType::Source);
+
+										if (WrittenBytes != Cmd.Size)
+										{
+											bOk = false;
+											UNSYNC_FATAL(L"Expected to write %llu bytes, but written %llu", Cmd.Size, WrittenBytes);
+											Error.Set(AppError(L"Failed to write output"));
+										}
+									}
+									else
+									{
+										NumHashMismatches++;
+									}
+								}
+								else
+								{
+									UNSYNC_FATAL(L"Failed to decompress downloaded block");
+									Error.Set(AppError(L"Failed to decompress downloaded block"));
+								}
+							}
+
+							if (bOk)
+							{
+								auto ScatterIt = BlockScatterMap.find(BlockHash);
+								if (ScatterIt != BlockScatterMap.end())
+								{
+									const std::vector<uint64>& ScatterList = ScatterIt->second;
+									for (uint64 ScatterOffset : ScatterList)
+									{
+										uint64 WrittenBytes = Output.Write(DecompressedData.GetData(), ScatterOffset, Cmd.Size);
+
+										Stats.WrittenBytesFromSource += WrittenBytes;
+
+										if (WrittenBytes != Cmd.Size)
+										{
+											UNSYNC_FATAL(L"Expected to write %llu bytes, but written %llu", Cmd.Size, WrittenBytes);
+											Error.Set(AppError(L"Failed to write output"));
+										}
+									}
+
+									AddGlobalProgress(Cmd.Size * ScatterList.size(), EBlockListType::Source);
+								}
+							}
+
+							if (bOk)
+							{
+								std::lock_guard<std::mutex> LockGuard(DownloadedBlocksMutex);
+								DownloadedBlocks.insert(BlockHash);
+							}
+
+							DecompressionSemaphore.Release();
+						});
+				}
+			});
+
+		DecompressTasks.wait();
+
+		DownloadProgressLogger.Complete();
+
+		if (NumHashMismatches.load() != 0)
+		{
+			UNSYNC_WARNING(L"Found block hash mismatches while downloading data from proxy. Mismatching blocks: %llu.",
+						   NumHashMismatches.load());
+		}
+
+		{
+			auto FilterResult = std::remove_if(FilteredSourceNeedList.begin(),
+											   FilteredSourceNeedList.end(),
+											   [&DownloadedBlocks](const FNeedBlock& Block)
+											   { return DownloadedBlocks.find(Block.Hash.ToHash128()) != DownloadedBlocks.end(); });
+			FilteredSourceNeedList.erase(FilterResult, FilteredSourceNeedList.end());
+		}
+
+		if (!FilteredSourceNeedList.empty())
+		{
+			uint64 RemainingBytes = ComputeSize(FilteredSourceNeedList);
+
+			UNSYNC_VERBOSE(L"Did not receive %d blocks from proxy. Must read %.2f MB from source.",
+						   FilteredSourceNeedList.size(),
+						   SizeMb(RemainingBytes));
+
+			UNSYNC_LOG_INDENT;
+
+			ProcessNeedList(Source, FilteredSourceNeedList, RemainingBytes, EBlockListType::Source, bSourceDataCopyTaskDone);
+		}
+		else
+		{
+			bSourceDataCopyTaskDone = true;
+		}
+	}
+	else if (SizeInfo.SourceBytes)
+	{
+		ProcessNeedList(Source, FilteredSourceNeedList, SizeInfo.SourceBytes, EBlockListType::Source, bSourceDataCopyTaskDone);
+	}
+
+	if (!bBaseDataCopyTaskDone)
+	{
+		UNSYNC_VERBOSE(L"Reading blocks from base");
+	}
+
+	bWaitingForBaseData = true;
+	BackgroundTasks.wait();
+	WriteTasks.wait();
+
+	double Duration = DurationSec(TimeBegin, TimePointNow());
+	UNSYNC_VERBOSE(L"Done in %.3f sec (%.3f MB / sec)", Duration, SizeMb(double(SizeInfo.TotalBytes) / Duration));
+
+	if (GLogVerbose)
+	{
+		LogGlobalProgress();
+	}
+
+	BuildResult.bSuccess	= !Error;
+	BuildResult.BaseBytes	= Stats.WrittenBytesFromBase;
+	BuildResult.SourceBytes = Stats.WrittenBytesFromSource;
+
+	return BuildResult;
+}
+
+FBuffer
+BuildTargetBuffer(FIOReader& SourceProvider, FIOReader& BaseProvider, const FNeedList& NeedList, const FBuildTargetParams& Params)
+{
+	FBuffer				Result;
+	const FNeedListSize SizeInfo = ComputeNeedListSize(NeedList);
+	Result.Resize(SizeInfo.TotalBytes);
+	FMemReaderWriter ResultWriter(Result.Data(), Result.Size());
+	BuildTarget(ResultWriter, SourceProvider, BaseProvider, NeedList, Params);
+	return Result;
+}
+
+FBuffer
+BuildTargetBuffer(const uint8*				SourceData,
+				  uint64					SourceSize,
+				  const uint8*				BaseData,
+				  uint64					BaseSize,
+				  const FNeedList&			NeedList,
+				  const FBuildTargetParams& Params)
+{
+	FMemReader SourceReader(SourceData, SourceSize);
+	FMemReader BaseReader(BaseData, BaseSize);
+	return BuildTargetBuffer(SourceReader, BaseReader, NeedList, Params);
+}
+
+struct FMemStream
+{
+	FMemStream(const uint8* InData, uint64 InDataSize) : Data(InData), DataRw(nullptr), SIZE(InDataSize), Offset(0) {}
+
+	FMemStream(uint8* InData, uint64 InDataSize) : Data(InData), DataRw(InData), SIZE(InDataSize), Offset(0) {}
+
+	bool Read(void* Dest, uint64 ReadSize)
+	{
+		if (Offset + ReadSize <= SIZE)
+		{
+			uint8* DestBytes = reinterpret_cast<uint8*>(Dest);
+			memcpy(DestBytes, Data + Offset, ReadSize);
+			Offset += ReadSize;
+			return true;
+		}
+		else
+		{
+			return false;
+		}
+	}
+
+	bool Write(const void* Source, uint64 WriteSize)
+	{
+		if (Offset + WriteSize <= SIZE)
+		{
+			const uint8* SourceBytes = reinterpret_cast<const uint8*>(Source);
+			memcpy(DataRw + Offset, SourceBytes, WriteSize);
+			Offset += WriteSize;
+			return true;
+		}
+		else
+		{
+			return false;
+		}
+	}
+
+	const uint8* Data;
+	uint8*		 DataRw;
+	const uint64 SIZE;
+	uint64		 Offset;
+};
+
+template<typename BlockType>
+static bool
+BlocksEqual(const std::vector<BlockType>& A, std::vector<BlockType>& B)
+{
+	if (A.size() != B.size())
+	{
+		return false;
+	}
+
+	for (uint64 I = 0; I < A.size(); ++I)
+	{
+		if (A[I].HashStrong != B[I].HashStrong || A[I].HashWeak != B[I].HashWeak || A[I].Size != B[I].Size || A[I].Offset != B[I].Offset)
+		{
+			return false;
+		}
+	}
+
+	return true;
+}
+
+FBuffer
+BuildTargetWithPatch(const uint8* PatchData, uint64 PatchSize, const uint8* BaseData, uint64 BaseDataSize)
+{
+	const EChunkingAlgorithmID ChunkingMode = EChunkingAlgorithmID::FixedBlocks;
+
+	FBuffer Result;
+
+	FBuffer	   DecompressedPatch = Decompress(PatchData, PatchSize);
+	FMemStream Stream(DecompressedPatch.Data(), DecompressedPatch.Size());
+
+	FPatchHeader Header;
+	Stream.Read(&Header, sizeof(Header));
+
+	if (Header.Magic != FPatchHeader::MAGIC)
+	{
+		UNSYNC_ERROR(L"Patch file header mismatch. Expected %llX, got %llX", (long long)FPatchHeader::MAGIC, (long long)Header.Magic);
+		return Result;
+	}
+
+	if (Header.Version != FPatchHeader::VERSION)
+	{
+		UNSYNC_ERROR(L"Patch file version mismatch. Expected %lld, got %lld", (long long)FPatchHeader::VERSION, (long long)Header.Version);
+		return Result;
+	}
+
+	FHash128 ExpectedHeaderHash = HashBlake3Bytes<FHash128>(Stream.Data, Stream.Offset);
+	FHash128 HeaderHash;
+	Stream.Read(&HeaderHash, sizeof(HeaderHash));
+	if (HeaderHash != ExpectedHeaderHash)
+	{
+		UNSYNC_ERROR(L"Patch header hash mismatch");
+		return Result;
+	}
+
+	FGenericBlockArray SourceValidationExpected(Header.NumSourceValidationBlocks);
+	FGenericBlockArray BaseValidationExpected(Header.NumBaseValidationBlocks);
+
+	Stream.Read(&SourceValidationExpected[0], SourceValidationExpected.size() * sizeof(SourceValidationExpected[0]));
+	Stream.Read(&BaseValidationExpected[0], BaseValidationExpected.size() * sizeof(BaseValidationExpected[0]));
+
+	FGenericBlockArray BaseValidation;
+
+	{
+		FLogVerbosityScope VerbosityScope(false);
+		FAlgorithmOptions  Algorithm;
+		Algorithm.ChunkingAlgorithmId	= ChunkingMode;
+		Algorithm.WeakHashAlgorithmId	= Header.WeakHashAlgorithmId;
+		Algorithm.StrongHashAlgorithmId = Header.StrongHashAlgorithmId;
+
+		BaseValidation = ComputeBlocks(BaseData, BaseDataSize, FPatchHeader::VALIDATION_BLOCK_SIZE, Algorithm);
+	}
+
+	if (!BlocksEqual(BaseValidation, BaseValidationExpected))
+	{
+		UNSYNC_ERROR(L"Base file mismatch");
+		return Result;
+	}
+
+	FPatchCommandList PatchCommands;
+	PatchCommands.Source.resize(Header.NumSourceBlocks);
+	PatchCommands.Base.resize(Header.NumBaseBlocks);
+
+	Stream.Read(&PatchCommands.Source[0], Header.NumSourceBlocks * sizeof(PatchCommands.Source[0]));
+	Stream.Read(&PatchCommands.Base[0], Header.NumBaseBlocks * sizeof(PatchCommands.Base[0]));
+
+	FHash128 ExpectedBlockHash = HashBlake3Bytes<FHash128>(Stream.Data, Stream.Offset);
+	FHash128 BlockHash;
+	Stream.Read(&BlockHash, sizeof(BlockHash));
+	if (BlockHash != ExpectedBlockHash)
+	{
+		UNSYNC_ERROR(L"Patch block hash mismatch");
+		return Result;
+	}
+
+	const uint8* SourceData		= Stream.Data + Stream.Offset;
+	uint64		 SourceDataSize = Stream.SIZE - Stream.Offset;
+
+	// TODO: generate the proper source need list from the start
+	uint64 SourceOffset = 0;
+	for (FCopyCommand& Block : PatchCommands.Source)
+	{
+		Block.SourceOffset = SourceOffset;
+		SourceOffset += Block.Size;
+	}
+
+	FNeedList NeedList;
+	NeedList.Base.reserve(PatchCommands.Base.size());
+	NeedList.Source.reserve(PatchCommands.Source.size());
+	NeedList.Sequence.reserve(PatchCommands.Sequence.size());
+
+	for (const FCopyCommand& Cmd : PatchCommands.Base)
+	{
+		FNeedBlock Block;
+		static_cast<FCopyCommand&>(Block) = Cmd;
+		NeedList.Base.push_back(Block);
+	}
+
+	for (const FCopyCommand& Cmd : PatchCommands.Source)
+	{
+		FNeedBlock Block;
+		static_cast<FCopyCommand&>(Block) = Cmd;
+		NeedList.Source.push_back(Block);
+	}
+
+	for (const FHash128& Hash : PatchCommands.Sequence)
+	{
+		NeedList.Sequence.push_back(Hash);
+	}
+
+	FBuildTargetParams BuildParams;
+	BuildParams.StrongHasher = Header.StrongHashAlgorithmId;
+	BuildParams.SourceType	 = FBuildTargetParams::ESourceType::Patch;
+
+	Result = BuildTargetBuffer(SourceData, SourceDataSize, BaseData, BaseDataSize, NeedList, BuildParams);
+
+	FGenericBlockArray SourceValidation;
+
+	{
+		FLogVerbosityScope VerbosityScope(false);
+
+		FAlgorithmOptions Algorithm;
+		Algorithm.ChunkingAlgorithmId	= ChunkingMode;
+		Algorithm.WeakHashAlgorithmId	= Header.WeakHashAlgorithmId;
+		Algorithm.StrongHashAlgorithmId = Header.StrongHashAlgorithmId;
+
+		SourceValidation = ComputeBlocks(Result.Data(), Result.Size(), FPatchHeader::VALIDATION_BLOCK_SIZE, Algorithm);
+	}
+
+	if (!BlocksEqual(SourceValidation, SourceValidationExpected))
+	{
+		UNSYNC_ERROR(L"Patched size file mismatch");
+		Result.Clear();
+		return Result;
+	}
+
+	return Result;
+}
+
+}  // namespace unsync
diff --git a/Engine/Source/Programs/Unsync/Private/UnsyncTarget.h b/Engine/Source/Programs/Unsync/Private/UnsyncTarget.h
new file mode 100644
--- /dev/null
+++ b/Engine/Source/Programs/Unsync/Private/UnsyncTarget.h
@@ -0,0 +1,55 @@
+// Copyright Epic Games, Inc. All Rights Reserved.
+
+#pragma once
+
+#include "UnsyncBuffer.h"
+#include "UnsyncCommon.h"
+#include "UnsyncFile.h"
+
+namespace unsync {
+
+struct FNeedList;
+class FProxyPool;
+class FBlockCache;
+class FScavengeDatabase;
+
+struct FBuildTargetResult
+{
+	bool   bSuccess	   = false;
+	uint64 SourceBytes = 0;
+	uint64 BaseBytes   = 0;
+};
+
+struct FBuildTargetParams
+{
+	EStrongHashAlgorithmID StrongHasher;
+	FProxyPool*			   ProxyPool		= nullptr;
+	FBlockCache*		   BlockCache		= nullptr;
+	FScavengeDatabase*	   ScavengeDatabase = nullptr;
+
+	enum class ESourceType {
+		File,
+		Patch
+	};
+
+	ESourceType SourceType = ESourceType::File;
+};
+
+FBuildTargetResult BuildTarget(FIOWriter&				 Result,
+							   FIOReader&				 Source,
+							   FIOReader&				 Base,
+							   const FNeedList&			 NeedList,
+							   const FBuildTargetParams& Params);
+
+FBuffer BuildTargetBuffer(FIOReader& SourceProvider, FIOReader& BaseProvider, const FNeedList& NeedList, const FBuildTargetParams& Params);
+
+FBuffer BuildTargetBuffer(const uint8*				SourceData,
+						  uint64					SourceSize,
+						  const uint8*				BaseData,
+						  uint64					BaseSize,
+						  const FNeedList&			NeedList,
+						  const FBuildTargetParams& Params);
+
+FBuffer BuildTargetWithPatch(const uint8* PatchData, uint64 PatchSize, const uint8* BaseData, uint64 BaseSize);
+
+}  // namespace unsync
diff --git a/Engine/Source/Programs/Unsync/Private/UnsyncTest.cpp b/Engine/Source/Programs/Unsync/Private/UnsyncTest.cpp
--- a/Engine/Source/Programs/Unsync/Private/UnsyncTest.cpp
+++ b/Engine/Source/Programs/Unsync/Private/UnsyncTest.cpp
@@ -8,6 +8,7 @@
 #include "UnsyncTest.h"
 #include "UnsyncThread.h"
 #include "UnsyncUtil.h"
+#include "UnsyncTarget.h"
 
 UNSYNC_THIRD_PARTY_INCLUDES_START
 #include <md5-sse2.h>
