unsync - Increase default memory budget for background sync tasks and add option to set it explicitly

diff --git a/Engine/Source/Programs/Unsync/Private/UnsyncCmdSync.cpp b/Engine/Source/Programs/Unsync/Private/UnsyncCmdSync.cpp
--- a/Engine/Source/Programs/Unsync/Private/UnsyncCmdSync.cpp
+++ b/Engine/Source/Programs/Unsync/Private/UnsyncCmdSync.cpp
@@ -166,20 +166,21 @@ CmdSync(const FCmdSyncOptions& Options)
 				SyncOptions.SourceType = ESyncSourceType::Server;
 			}
 
-			SyncOptions.Source				   = ResolvedSource;
-			SyncOptions.Base				   = Options.Target;  // read base data from existing target
-			SyncOptions.Target				   = Options.Target;
-			SyncOptions.ScavengeRoot		   = Options.ScavengeRoot;
-			SyncOptions.ScavengeDepth		   = Options.ScavengeDepth;
-			SyncOptions.Overlays			   = ResolvedOverlays;
-			SyncOptions.SourceManifestOverride = Options.SourceManifestOverride;
-			SyncOptions.ProxyPool			   = &ProxyPool;
-			SyncOptions.SyncFilter			   = Options.Filter;
-			SyncOptions.bCleanup			   = Options.bCleanup;
-			SyncOptions.bValidateSourceFiles   = Options.bFullSourceScan;
-			SyncOptions.bFullDifference		   = Options.bFullDifference;
-			SyncOptions.bValidateTargetFiles   = Options.bValidateTargetFiles;
-			SyncOptions.bCheckAvailableSpace   = Options.bCheckAvailableSpace;
+			SyncOptions.Source					   = ResolvedSource;
+			SyncOptions.Base					   = Options.Target;  // read base data from existing target
+			SyncOptions.Target					   = Options.Target;
+			SyncOptions.ScavengeRoot			   = Options.ScavengeRoot;
+			SyncOptions.ScavengeDepth			   = Options.ScavengeDepth;
+			SyncOptions.Overlays				   = ResolvedOverlays;
+			SyncOptions.SourceManifestOverride	   = Options.SourceManifestOverride;
+			SyncOptions.ProxyPool				   = &ProxyPool;
+			SyncOptions.SyncFilter				   = Options.Filter;
+			SyncOptions.bCleanup				   = Options.bCleanup;
+			SyncOptions.bValidateSourceFiles	   = Options.bFullSourceScan;
+			SyncOptions.bFullDifference			   = Options.bFullDifference;
+			SyncOptions.bValidateTargetFiles	   = Options.bValidateTargetFiles;
+			SyncOptions.bCheckAvailableSpace	   = Options.bCheckAvailableSpace;
+			SyncOptions.BackgroundTaskMemoryBudget = Options.BackgroundTaskMemoryBudget;
 
 			return SyncDirectory(SyncOptions) ? 0 : 1;
 		}
diff --git a/Engine/Source/Programs/Unsync/Private/UnsyncCmdSync.h b/Engine/Source/Programs/Unsync/Private/UnsyncCmdSync.h
--- a/Engine/Source/Programs/Unsync/Private/UnsyncCmdSync.h
+++ b/Engine/Source/Programs/Unsync/Private/UnsyncCmdSync.h
@@ -17,6 +17,7 @@ struct FCmdSyncOptions
 	FPath SourceManifestOverride;
 	FPath ScavengeRoot;
 	uint32 ScavengeDepth = 5;
+	uint64 BackgroundTaskMemoryBudget = 2_GB;
 
 	std::vector<FPath> Overlays;
 
diff --git a/Engine/Source/Programs/Unsync/Private/UnsyncCore.cpp b/Engine/Source/Programs/Unsync/Private/UnsyncCore.cpp
--- a/Engine/Source/Programs/Unsync/Private/UnsyncCore.cpp
+++ b/Engine/Source/Programs/Unsync/Private/UnsyncCore.cpp
@@ -27,7 +27,7 @@ UNSYNC_THIRD_PARTY_INCLUDES_START
 #include <md5-sse2.h>
 UNSYNC_THIRD_PARTY_INCLUDES_END
 
-#define UNSYNC_VERSION_STR "1.0.61"
+#define UNSYNC_VERSION_STR "1.0.62"
 
 namespace unsync {
 
@@ -2685,9 +2685,11 @@ SyncDirectory(const FSyncDirectoryOptions& SyncOptions)
 
 	{
 		// Throttle background tasks by trying to keep them to some sensible memory budget. Best effort only, not a hard limit.
-		static constexpr uint64 BackgroundTaskMemoryBudget	= 2_GB;
-		static constexpr uint64 TargetTotalSizePerTaskBatch = BackgroundTaskMemoryBudget;
-		static constexpr uint64 MaxFilesPerTaskBatch		= 1000;
+		const uint64 BackgroundTaskMemoryBudget	 = SyncOptions.BackgroundTaskMemoryBudget;
+		const uint64 TargetTotalSizePerTaskBatch = BackgroundTaskMemoryBudget;
+		const uint64 MaxFilesPerTaskBatch		 = 1000;
+
+		UNSYNC_VERBOSE2(L"Background task memory budget: %llu GB", BackgroundTaskMemoryBudget >> 30);
 
 		struct FBackgroundTaskResult
 		{
diff --git a/Engine/Source/Programs/Unsync/Private/UnsyncCore.h b/Engine/Source/Programs/Unsync/Private/UnsyncCore.h
--- a/Engine/Source/Programs/Unsync/Private/UnsyncCore.h
+++ b/Engine/Source/Programs/Unsync/Private/UnsyncCore.h
@@ -328,7 +328,8 @@ struct FSyncDirectoryOptions
 	bool			   bValidateSourceFiles = true;	 // whether to check that all source files declared in the manifest are present/valid
 	bool			   bValidateTargetFiles = true;	 // WARNING: turning this off is intended only for testing/profiling
 	bool			   bFullDifference = true;	// whether to run full file difference algorithm, even when there is an existing manifest
-	bool			   bCheckAvailableSpace = true;	 // whether to abort the sync if target path does not have enough available space
+	bool			   bCheckAvailableSpace		  = true;  // whether to abort the sync if target path does not have enough available space
+	uint64			   BackgroundTaskMemoryBudget = 2_GB;
 };
 
 bool SyncDirectory(const FSyncDirectoryOptions& SyncOptions);
diff --git a/Engine/Source/Programs/Unsync/Private/UnsyncMain.cpp b/Engine/Source/Programs/Unsync/Private/UnsyncMain.cpp
--- a/Engine/Source/Programs/Unsync/Private/UnsyncMain.cpp
+++ b/Engine/Source/Programs/Unsync/Private/UnsyncMain.cpp
@@ -104,6 +104,7 @@ InnerMain(int Argc, char** Argv)
 	int32					 CompressionLevel	 = 3;
 	uint32					 DiffBlockSize		 = uint32(4_KB);
 	uint32					 HashOrSyncBlockSize = uint32(64_KB);
+	uint32					 BackgroundTaskMemoryBudgetGB = 2;
 
 	struct FDeprecatedOptions
 	{
@@ -303,6 +304,10 @@ InnerMain(int Argc, char** Argv)
 	SubSync->add_flag("--login", bShouldLogin, "Use user authentication when accessing unsync server");
 	SubSync->add_flag("--no-timeout", bNoSocketTimeout, "Disable the default 60 second timeout on network socket operations");
 
+	CLI::Option* BackgroundMemoryBudgetOption = SubSync->add_option("--background-task-memory",
+															BackgroundTaskMemoryBudgetGB,
+															"Set memory budget that background tasks in gigabytes (default: 2 GB)");
+
 	SubCommands.push_back(SubSync);
 
 	CLI::App* SubPatch = Cli.add_subcommand("patch", "Applies a patch generated with 'diff' on top of base file");
@@ -858,20 +863,41 @@ InnerMain(int Argc, char** Argv)
 			RemoteDesc.RecvTimeoutSeconds = 60;
 		}
 
+		// Try to derive default memory budget
+
+		if (BackgroundMemoryBudgetOption->empty())
+		{
+			FSystemMemoryInfo MemoryInfo;
+			if (QueryMemoryInfo(MemoryInfo))
+			{
+				uint32 InstalledMemoryGB = CheckedNarrow(MemoryInfo.InstalledPhysicalMemory >> 30);
+				UNSYNC_VERBOSE2(L"Detected memory: %llu GB", InstalledMemoryGB);
+				BackgroundTaskMemoryBudgetGB = std::max<uint32>(2, InstalledMemoryGB / 4);
+			}
+
+			UNSYNC_VERBOSE2(L"Using automatic background task memory budget: %llu GB", BackgroundTaskMemoryBudgetGB);
+		}
+		else
+		{
+			UNSYNC_VERBOSE2(L"Using explicit background task memory budget: %llu GB", BackgroundTaskMemoryBudgetGB);
+		}
+		//
+
 		FCmdSyncOptions SyncOptions;
 
-		SyncOptions.Algorithm			   = Algorithm;
-		SyncOptions.Source				   = SourceFilename;
-		SyncOptions.Target				   = TargetFilename;
-		SyncOptions.SourceManifestOverride = SourceManifestFilename;
-		SyncOptions.Remote				   = RemoteDesc;
-		SyncOptions.bFullDifference		   = bFullDifference;
-		SyncOptions.bFullSourceScan		   = bFullSourceScan;
-		SyncOptions.bCleanup			   = !bNoCleanupAfterSync;
-		SyncOptions.Filter				   = &SyncFilter;
-		SyncOptions.bValidateTargetFiles   = !bNoOutputValidation;
-		SyncOptions.bCheckAvailableSpace   = !bNoSpaceValidation;
-		SyncOptions.ScavengeRoot		   = ScavengeRoot;
+		SyncOptions.Algorithm				   = Algorithm;
+		SyncOptions.Source					   = SourceFilename;
+		SyncOptions.Target					   = TargetFilename;
+		SyncOptions.SourceManifestOverride	   = SourceManifestFilename;
+		SyncOptions.Remote					   = RemoteDesc;
+		SyncOptions.bFullDifference			   = bFullDifference;
+		SyncOptions.bFullSourceScan			   = bFullSourceScan;
+		SyncOptions.bCleanup				   = !bNoCleanupAfterSync;
+		SyncOptions.Filter					   = &SyncFilter;
+		SyncOptions.bValidateTargetFiles	   = !bNoOutputValidation;
+		SyncOptions.bCheckAvailableSpace	   = !bNoSpaceValidation;
+		SyncOptions.ScavengeRoot			   = ScavengeRoot;
+		SyncOptions.BackgroundTaskMemoryBudget = uint64(BackgroundTaskMemoryBudgetGB) << 30ull;
 
 		for (const std::string& Entry : OverlayArrayUtf8)
 		{
diff --git a/Engine/Source/Programs/Unsync/Private/UnsyncMemory.cpp b/Engine/Source/Programs/Unsync/Private/UnsyncMemory.cpp
--- a/Engine/Source/Programs/Unsync/Private/UnsyncMemory.cpp
+++ b/Engine/Source/Programs/Unsync/Private/UnsyncMemory.cpp
@@ -4,6 +4,10 @@
 #include "UnsyncLog.h"
 #include "UnsyncUtil.h"
 
+#if UNSYNC_PLATFORM_WINDOWS
+#	include <Windows.h>
+#endif	// UNSYNC_PLATFORM_WINDOWS
+
 #if UNSYNC_USE_DEBUG_HEAP
 #	include "ig-debugheap/DebugHeap.h"
 #endif	// UNSYNC_USE_DEBUG_HEAP
@@ -224,6 +228,33 @@ UnsyncDebugFree(void*)
 
 #endif	// UNSYNC_USE_DEBUG_HEAP
 
+
+#if UNSYNC_PLATFORM_WINDOWS
+bool
+QueryMemoryInfo(FSystemMemoryInfo& OutMemoryInfo)
+{
+	OutMemoryInfo = {};
+
+	ULONGLONG TotalMemoryInKB = 0;
+
+	if (!GetPhysicallyInstalledSystemMemory(&TotalMemoryInKB))
+	{
+		return false;
+	}
+
+	OutMemoryInfo.InstalledPhysicalMemory = uint64(TotalMemoryInKB) << 10ull;
+
+	return true;
+}
+#else
+bool
+QueryMemoryInfo(FSystemMemoryInfo& OutMemoryInfo)
+{
+	OutMemoryInfo = {};
+	return false;
+}
+#endif
+
 }  // namespace unsync
 
 #if UNSYNC_OVERRIDE_GLOBAL_NEW_DELETE
@@ -243,8 +274,16 @@ IsBootstrapPtr(void* InPtr)
 }
 }  // namespace unsync
 
+#	if _MSC_VER
+_NODISCARD _Ret_notnull_
+_Post_writable_byte_size_(Size)
+_VCRT_ALLOCATOR
+void* __CRTDECL
+operator new(size_t Size)
+#	else
 void*
 operator new(size_t Size)
+#	endif
 {
 	using namespace unsync;
 
diff --git a/Engine/Source/Programs/Unsync/Private/UnsyncMemory.h b/Engine/Source/Programs/Unsync/Private/UnsyncMemory.h
--- a/Engine/Source/Programs/Unsync/Private/UnsyncMemory.h
+++ b/Engine/Source/Programs/Unsync/Private/UnsyncMemory.h
@@ -23,4 +23,12 @@ void  UnsyncMallocInit(EMallocType MallocType);
 void* UnsyncMalloc(size_t Size);
 void  UnsyncFree(void* Ptr);
 
+struct FSystemMemoryInfo
+{
+	uint64 InstalledPhysicalMemory = 0;
+};
+
+// Queries various memory properties from the OS, returns false if information is not available
+bool QueryMemoryInfo(FSystemMemoryInfo& OutMemoryInfo);
+
 }  // namespace unsync
