unsync - Fix compile error on Linux

diff --git a/Engine/Source/Programs/Unsync/Private/UnsyncMemory.cpp b/Engine/Source/Programs/Unsync/Private/UnsyncMemory.cpp
--- a/Engine/Source/Programs/Unsync/Private/UnsyncMemory.cpp
+++ b/Engine/Source/Programs/Unsync/Private/UnsyncMemory.cpp
@@ -274,7 +274,7 @@ IsBootstrapPtr(void* InPtr)
 }
 }  // namespace unsync
 
-#	if _MSC_VER
+#	if defined(_MSC_VER)
 _NODISCARD _Ret_notnull_
 _Post_writable_byte_size_(Size)
 _VCRT_ALLOCATOR
