unsync - Unpack command now outputs revision control information into a dedicated file

diff --git a/Engine/Source/Programs/Unsync/Private/UnsyncCmdPack.cpp b/Engine/Source/Programs/Unsync/Private/UnsyncCmdPack.cpp
--- a/Engine/Source/Programs/Unsync/Private/UnsyncCmdPack.cpp
+++ b/Engine/Source/Programs/Unsync/Private/UnsyncCmdPack.cpp
@@ -143,7 +143,11 @@ struct FPackDatabase
 				FEntry DatabaseEntry;
 				DatabaseEntry.IndexEntry	   = IndexEntry;
 				DatabaseEntry.PackIndex		   = PackIndex;
-				BlockMap[IndexEntry.BlockHash] = DatabaseEntry;
+
+				if (BlockMap.insert(std::make_pair(IndexEntry.BlockHash, DatabaseEntry)).second)
+				{
+					TotalCompressedSize += DatabaseEntry.IndexEntry.CompressedSize;
+				}
 			}
 		}
 	}
@@ -164,6 +168,7 @@ struct FPackDatabase
 
 	std::vector<FPath>						  PackFilenames;
 	THashMap<FHash128, FEntry>				  BlockMap;
+	uint64									  TotalCompressedSize = 0;
 
 	mutable std::vector<std::shared_ptr<FNativeFile>> PackFileCache;
 	mutable std::mutex FileCacheMutex;
@@ -896,6 +901,76 @@ SyncDirectoryFromPack(const FPath& OutputRoot, const FPackDatabase& PackDb, cons
 	return !Error;
 }
 
+bool
+VerifyManifest(const FPackDatabase& PackDb, const FDirectoryManifest& NewDirectoryManifest)
+{
+	FAtomicError Error;
+
+	auto VerifyFileManifest = [&Error, &PackDb](const std::pair<const std::wstring, FFileManifest>& FileIt)
+	{
+		if (Error)
+		{
+			return;
+		}
+
+		const FFileManifest& FileManifest = FileIt.second;
+		for (const FGenericBlock& Block : FileManifest.Blocks)
+		{
+			if (PackDb.BlockMap.find(Block.HashStrong.ToHash128()) == PackDb.BlockMap.end())
+			{
+				Error.Set(AppError("Found unknown block in the manifest"));
+			}
+		}
+	};
+
+	ParallelForEach(NewDirectoryManifest.Files, VerifyFileManifest);
+
+	return !Error;
+}
+
+enum class ERevisionControlFileFormat
+{
+	IdentityOnly, // only include revision control identity of the file
+	P4Have,  // use the same format as `p4 have` output, i.e. "//depot/file.txt#123 - C:\Local\Path\file.txt"
+};
+
+bool
+SaveRevisionControlData(const FPath& OutputPath, const FDirectoryManifest& Manifest, ERevisionControlFileFormat Format)
+{
+	if (OutputPath.has_parent_path())
+	{
+		FPath ParentPath = OutputPath.parent_path();
+		if (!EnsureDirectoryExists(ParentPath))
+		{
+			UNSYNC_ERROR(L"Failed to create output directory '%ls'", ParentPath.wstring().c_str());
+			return false;
+		}
+	}
+
+	std::string OutputString;
+
+	std::string CurrentPathUtf8;
+	for (const auto& It : Manifest.Files)
+	{
+		const FFileManifest& FileManifest	 = It.second;
+		const std::string&	 Identity		 = FileManifest.RevisionControlIdentity;
+
+		ConvertWideToUtf8(FileManifest.CurrentPath.wstring(), CurrentPathUtf8);
+
+		OutputString.append(Identity);
+
+		if (Format == ERevisionControlFileFormat::P4Have)
+		{
+			OutputString.append(" - ");
+			OutputString.append(CurrentPathUtf8);
+		}
+
+		OutputString.append("\n");
+	}
+
+	return WriteBufferToFile(OutputPath, OutputString);
+}
+
 int32
 CmdUnpack(const FCmdUnpackOptions& Options)
 {
@@ -911,11 +986,22 @@ CmdUnpack(const FCmdUnpackOptions& Options)
 	const FPath& StoreRoot = Options.StorePath;
 	const FPath	 PackRoot  = StoreRoot / "pack";
 
+	const FPath ManifestRoot		  = Options.OutputPath / ".unsync";
+	const FPath DirectoryManifestPath = ManifestRoot / "manifest.bin";
+	const FPath RevisionFilePath	  = ManifestRoot / "revisions.txt";
+
 	FPackDatabase PackDb;
 	{
 		UNSYNC_LOG(L"Loading block database");
 		UNSYNC_LOG_INDENT;
+
+		// TODO: sync index files to local cache for faster unpack next time
 		PackDb.Load(PackRoot);
+		UNSYNC_LOG("Known pack files: %llu", llu(PackDb.PackFilenames.size()));
+		UNSYNC_LOG("Known blocks: %llu", llu(PackDb.BlockMap.size()));
+		UNSYNC_LOG("Total compressed size: %llu (%.3f GB)",
+				   llu(PackDb.TotalCompressedSize),
+				   double(PackDb.TotalCompressedSize) / double(1 << 30));
 	}
 
 	UNSYNC_LOG(L"Reading snapshot");
@@ -953,11 +1039,10 @@ CmdUnpack(const FCmdUnpackOptions& Options)
 			return -1;
 		}
 	}
-	
 
 	FDirectoryManifest NewDirectoryManifest;
 	{
-		UNSYNC_LOG(L"Loading directory manifest");
+		UNSYNC_LOG(L"Parsing directory manifest");
 
 		FMemReader		   ManifestMemReader(ManifestBuffer);
 		FIOReaderStream	   ManifestReaderStream(ManifestMemReader);
@@ -968,22 +1053,83 @@ CmdUnpack(const FCmdUnpackOptions& Options)
 		}
 	}
 
-	UNSYNC_LOG(L"Writing target files");
 	{
-		UNSYNC_LOG_INDENT;
-		if (!SyncDirectoryFromPack(Options.OutputPath, PackDb, NewDirectoryManifest))
+		UNSYNC_LOG(L"Verifying manifest against known block database");
+		if (!VerifyManifest(PackDb, NewDirectoryManifest))
 		{
-			UNSYNC_FATAL(L"Failed to sync directory pack");
+			UNSYNC_FATAL(L"Failed to verify manifest");
 			return -1;
 		}
 	}
 
-	UNSYNC_LOG(L"Saving directory manifest");
-	const FPath ManifestRoot		  = Options.OutputPath / ".unsync";
-	const FPath DirectoryManifestPath = ManifestRoot / "manifest.bin";
-	if (!GDryRun && EnsureDirectoryExists(ManifestRoot))
+	if (Options.bOutputFiles)
 	{
-		SaveDirectoryManifest(NewDirectoryManifest, DirectoryManifestPath);
+		UNSYNC_LOG(L"Writing target files");
+		{
+			UNSYNC_LOG_INDENT;
+			if (!SyncDirectoryFromPack(Options.OutputPath, PackDb, NewDirectoryManifest))
+			{
+				UNSYNC_FATAL(L"Failed to sync directory pack");
+				return -1;
+			}
+		}
+	}
+
+	{
+		UNSYNC_LOG(L"Saving directory manifest");
+		UNSYNC_LOG_INDENT;
+		if (!GDryRun && EnsureDirectoryExists(ManifestRoot))
+		{
+			SaveDirectoryManifest(NewDirectoryManifest, DirectoryManifestPath);
+		}
+	}
+
+	if (NewDirectoryManifest.bHasFileRevisionControl)
+	{
+		if (Options.bOutputRevisions)
+		{
+			UNSYNC_LOG(L"Extracting revision control identities");
+			UNSYNC_LOG_INDENT;
+
+			if (!SaveRevisionControlData(RevisionFilePath, NewDirectoryManifest, ERevisionControlFileFormat::IdentityOnly))
+			{
+				UNSYNC_FATAL(L"Failed to save revision control data");
+				return -1;
+			}
+		}
+
+		if (!Options.P4HaveOutputPath.empty())
+		{
+			UNSYNC_LOG(L"Extracting revision control data into `p4 have` file '%ls'", Options.P4HaveOutputPath.wstring().c_str());
+			UNSYNC_LOG_INDENT;
+
+			if (!SaveRevisionControlData(Options.P4HaveOutputPath, NewDirectoryManifest, ERevisionControlFileFormat::P4Have))
+			{
+				UNSYNC_FATAL(L"Failed to save revision control data");
+				return -1;
+			}
+		}
+	}
+	else
+	{
+		if (!Options.P4HaveOutputPath.empty())
+		{
+			UNSYNC_ERROR(L"P4 have output is requested, but the manifest does not contain revision control data");
+			return -1;
+		}
+
+		// Delete previously cached revision control file if current manifest doesn't have revision data
+		if (!GDryRun)
+		{
+			if (PathExists(RevisionFilePath))
+			{
+				std::error_code ErrorCode;
+				if (!FileRemove(RevisionFilePath, ErrorCode))
+				{
+					UNSYNC_WARNING(L"Failed to delete file '%ls'", RevisionFilePath.wstring().c_str())
+				}
+			}
+		}
 	}
 
 	return 0;
diff --git a/Engine/Source/Programs/Unsync/Private/UnsyncCmdPack.h b/Engine/Source/Programs/Unsync/Private/UnsyncCmdPack.h
--- a/Engine/Source/Programs/Unsync/Private/UnsyncCmdPack.h
+++ b/Engine/Source/Programs/Unsync/Private/UnsyncCmdPack.h
@@ -21,9 +21,14 @@ int32 CmdPack(const FCmdPackOptions& Options);
 
 struct FCmdUnpackOptions
 {
-	FPath			  OutputPath;
-	FPath			  StorePath;
-	std::string		  SnapshotName;
+	FPath		OutputPath;
+	FPath		P4HaveOutputPath;  // extract embedded p4 have data into the specified file
+	FPath		StorePath;
+	std::string SnapshotName;
+	bool		bOutputFiles	 = true;  // set to false to skip unpack process for files (useful for extracting metadata)
+	bool		bOutputRevisions = true;  // write [output]/.unsync/revisions.txt if manifest contains revision control data
+
+	// bool			  bRunP4Sync = false; // TODO: run batches of `p4 sync -k` as files get finalized
 };
 
 int32 CmdUnpack(const FCmdUnpackOptions& Options);
diff --git a/Engine/Source/Programs/Unsync/Private/UnsyncCore.cpp b/Engine/Source/Programs/Unsync/Private/UnsyncCore.cpp
--- a/Engine/Source/Programs/Unsync/Private/UnsyncCore.cpp
+++ b/Engine/Source/Programs/Unsync/Private/UnsyncCore.cpp
@@ -27,7 +27,7 @@ UNSYNC_THIRD_PARTY_INCLUDES_START
 #include <md5-sse2.h>
 UNSYNC_THIRD_PARTY_INCLUDES_END
 
-#define UNSYNC_VERSION_STR "1.0.62"
+#define UNSYNC_VERSION_STR "1.0.63-dev1"
 
 namespace unsync {
 
diff --git a/Engine/Source/Programs/Unsync/Private/UnsyncMain.cpp b/Engine/Source/Programs/Unsync/Private/UnsyncMain.cpp
--- a/Engine/Source/Programs/Unsync/Private/UnsyncMain.cpp
+++ b/Engine/Source/Programs/Unsync/Private/UnsyncMain.cpp
@@ -101,6 +101,8 @@ InnerMain(int Argc, char** Argv)
 	bool					 bQuickLogin		 = false;
 	bool					 bForceRefreshAuth	 = false;
 	bool					 bNoSocketTimeout	 = false;
+	bool					 bNoOutputFiles  	 = false;
+	bool					 bNoOutputRevisions  = false;
 	int32					 CompressionLevel	 = 3;
 	uint32					 DiffBlockSize		 = uint32(4_KB);
 	uint32					 HashOrSyncBlockSize = uint32(64_KB);
@@ -182,14 +184,19 @@ InnerMain(int Argc, char** Argv)
 
 	CLI::App* SubUnpack = nullptr;
 	{
-		SubUnpack =
-			Cli.add_subcommand("unpack",
-							   "EXPERIMENTAL: Sync directory based on package snapshot");
+		SubUnpack = Cli.add_subcommand("unpack", "EXPERIMENTAL: Sync directory based on package snapshot");
 
 		SubUnpack->add_option("Output", OutputFilenameUtf8, "Output directory path")->required();
 
 		SubUnpack->add_option("--store", StorePathUtf8, "Pack storage path")->required();
 		SubUnpack->add_option("--snapshot", SnapshotNameUtf8, "Directory snapshot ID")->required();
+		SubUnpack->add_option("--p4havefile", P4HavePathUtf8, "Write revision control data in `p4 have` format into this file");
+		SubUnpack->add_flag("--no-revisions", bNoOutputRevisions, "Skip writing revision control data to <output>/.unsync/revisions.txt");
+		SubUnpack->add_flag("--no-files",
+							bNoOutputFiles,
+							"Skip actually unpacking the snapshot files, but attempt to reconstruct and verify the manifest. "
+							"Can be used in combination with --p4havefile option to only extract the `p4 have` list.");
+
 
 		SubCommands.push_back(SubUnpack);
 	}
@@ -815,9 +822,12 @@ InnerMain(int Argc, char** Argv)
 	{
 		FCmdUnpackOptions UnpackOptions;
 
-		UnpackOptions.OutputPath   = OutputFilename;
-		UnpackOptions.SnapshotName = SnapshotNameUtf8;
-		UnpackOptions.StorePath	   = NormalizeFilenameUtf8(StorePathUtf8);
+		UnpackOptions.OutputPath	   = OutputFilename;
+		UnpackOptions.SnapshotName	   = SnapshotNameUtf8;
+		UnpackOptions.P4HaveOutputPath = NormalizeFilenameUtf8(P4HavePathUtf8);
+		UnpackOptions.StorePath		   = NormalizeFilenameUtf8(StorePathUtf8);
+		UnpackOptions.bOutputFiles	   = !bNoOutputFiles;
+		UnpackOptions.bOutputRevisions = !bNoOutputRevisions;
 
 		return CmdUnpack(UnpackOptions);
 	}
