unsync - Implement basic incremental sync in unpack command and named snapshots

diff --git a/Engine/Source/Programs/Unsync/Private/UnsyncCmdPack.cpp b/Engine/Source/Programs/Unsync/Private/UnsyncCmdPack.cpp
--- a/Engine/Source/Programs/Unsync/Private/UnsyncCmdPack.cpp
+++ b/Engine/Source/Programs/Unsync/Private/UnsyncCmdPack.cpp
@@ -7,6 +7,7 @@
 #include "UnsyncSerialization.h"
 #include "UnsyncThread.h"
 #include "UnsyncError.h"
+#include "UnsyncProgress.h"
 
 #include <stdio.h>
 #include <stdlib.h>
@@ -469,7 +470,7 @@ CmdPack(const FCmdPackOptions& Options)
 
 		UNSYNC_LOG(L"Loaded entries from p4 manifest: %llu", llu(DirectoryManifest.Files.size()));
 
-		UNSYNC_LOG(L"Updating file attributes");
+		UNSYNC_LOG(L"Updating file metadata");
 		auto UpdateFileMetadata = [&OldDirectoryManifest](std::pair<const std::wstring, FFileManifest>& It)
 		{
 			if (OldDirectoryManifest.bHasFileRevisionControl)
@@ -657,7 +658,7 @@ CmdPack(const FCmdPackOptions& Options)
 		FHash128 ManifestBlocksBufferHash = HashBlake3Bytes<FHash128>(ManifestBlocksBuffer.Data, ManifestBlocksBuffer.Size);
 
 		std::string SnapshotId	 = HashToHexString(ManifestBlocksBufferHash);  // TODO: allow overriding this from command line
-		FPath		SnapshotPath = StoreRoot / (SnapshotId + ".unsync_snapshot");
+		FPath		SnapshotPath = StoreRoot / "snapshot" / (SnapshotId + ".unsync_snapshot");
 
 		UNSYNC_LOG(L"Writing snapshot: %hs", SnapshotId.c_str());
 
@@ -665,6 +666,22 @@ CmdPack(const FCmdPackOptions& Options)
 		if (!bSnapshotWritten)
 		{
 			UNSYNC_FATAL(L"Failed to write snapthot file '%ls'", SnapshotPath.wstring().c_str());
+			return -1;
+		}
+
+		// If explicit snapshot name is provided, save it as a separate "tag" file,
+		// never overwriting normal snapshots which are named based on content hash.
+		if (!Options.SnapshotName.empty())
+		{
+			UNSYNC_LOG(L"Saving snapshot %hs as tag '%hs'", SnapshotId.c_str(), Options.SnapshotName.c_str());
+			FPath			TagPath = StoreRoot / "tag" / (Options.SnapshotName + ".unsync_tag");
+			std::error_code ErrorCode;
+			if (!FileCopyOverwrite(SnapshotPath, TagPath, ErrorCode))
+			{
+				std::string SystemError = FormatSystemErrorMessage(ErrorCode.value());
+				UNSYNC_FATAL(L"Failed to save snapshot tag '%ls'. %hs", TagPath.wstring().c_str(), SystemError.c_str());
+				return -1;
+			}
 		}
 
 		ManifestUniqueBytes += ProcessedRawBytes.load();
@@ -827,26 +844,157 @@ struct FDirectoryCreationCache
 	std::mutex					 Mutex;
 };
 
+uint64 ComputeManifestTotalSize(const FDirectoryManifest& Manifest)
+{
+	uint64 Result = 0;
+
+	for (const auto& It : Manifest.Files)
+	{
+		Result += It.second.Size;
+	}
+
+	return Result;
+}
+
 bool
 SyncDirectoryFromPack(const FPath& OutputRoot, const FPackDatabase& PackDb, const FDirectoryManifest& NewDirectoryManifest)
 {
-	// TODO: incremental sync
+	// TODO: incrementally patch individual files
 	// TODO: fetch blocks from server
 
+	FTimingLogger SyncTimingLogger("Sync time", ELogLevel::Info);
+
+	FDirectoryManifest OldDirectoryManifest;
+	FPath			   OldManifestPath = OutputRoot / ".unsync" / "manifest.bin";
+	if (PathExists(OldManifestPath))
+	{
+		UNSYNC_LOG(L"Loading existing unsync manifest for '%ls'", OutputRoot.wstring().c_str());
+		UNSYNC_LOG_INDENT;
+
+		if (LoadDirectoryManifest(OldDirectoryManifest, OutputRoot, OldManifestPath))
+		{
+			if (OldDirectoryManifest.bHasFileRevisionControl)
+			{
+				UNSYNC_LOG(L"Loaded existing manifest with revision control data");
+			}
+			else
+			{
+				UNSYNC_LOG(L"Loaded existing manifest without revision control data");
+			}
+
+			UNSYNC_LOG(L"Updating file metadata");
+			FTimingLogger TimingLogger("Manifest generation time", ELogLevel::Info);
+
+			std::atomic<uint64> NumDirtyFiles = 0;
+
+			FLogProgressScope ScanProgressLogger(OldDirectoryManifest.Files.size(), ELogProgressUnits::Raw, 1000, false /*bVerboseOnly*/);
+
+			FThreadLogConfig LogConfig;
+			auto UpdateFileMetadata = [&NumDirtyFiles, &ScanProgressLogger, &LogConfig](std::pair<const std::wstring, FFileManifest>& It)
+			{
+				FThreadLogConfig::FScope LogConfigScope(LogConfig);
+
+				FFileAttributes Attrib = GetFileAttrib(It.second.CurrentPath);
+
+				// Check if the file on disk still matches the manifest and invalidate it otherwise
+				if (!Attrib.bValid || It.second.Mtime != Attrib.Mtime || It.second.Size != Attrib.Size ||
+					It.second.bReadOnly != Attrib.bReadOnly)
+				{
+					It.second = FFileManifest();
+					NumDirtyFiles++;
+				}
+
+				ScanProgressLogger.Add(1);
+			};
+			ParallelForEach(OldDirectoryManifest.Files, UpdateFileMetadata);
+
+			ScanProgressLogger.Complete();
+
+			if (NumDirtyFiles)
+			{
+				UNSYNC_LOG(L"Found dirty files: %llu", llu(NumDirtyFiles));
+			}
+		}
+	}
+
+	if (!OldDirectoryManifest.IsValid())
+	{
+		UNSYNC_LOG(L"Creating unsync manifest for '%ls'", OutputRoot.wstring().c_str());
+		UNSYNC_LOG_INDENT;
+
+		FTimingLogger TimingLogger("Manifest generation complete", ELogLevel::Info);
+
+		FComputeBlocksParams LightweightManifestParams;
+		LightweightManifestParams.Algorithm	  = NewDirectoryManifest.Algorithm;
+		LightweightManifestParams.bNeedBlocks = false;
+		LightweightManifestParams.BlockSize	  = 0;
+
+		OldDirectoryManifest = CreateDirectoryManifest(OutputRoot, LightweightManifestParams);
+	}
+
+	if (!OldDirectoryManifest.IsValid())
+	{
+		UNSYNC_ERROR(L"Failed to load or create manifest for target directory");
+		return false;
+	}
+
+	UNSYNC_LOG("Writing target files");
+
 	FAtomicError Error;
 
 	FDirectoryCreationCache DirCache;
 
-	auto BuildTargetCallback = [&Error, &PackDb, &DirCache](const std::pair<const std::wstring, FFileManifest>& FileIt)
+	struct FStats
+	{
+		std::atomic<uint64> NumFilesSkipped;
+		std::atomic<uint64> NumFilesSynced;
+
+		std::atomic<uint64> BytesSkipped;
+		std::atomic<uint64> BytesSynced;
+	};
+
+	FStats Stats;
+
+	const uint64 TotalSizeInBytes = ComputeManifestTotalSize(NewDirectoryManifest);
+
+	FLogProgressScope SyncProgressLogger(TotalSizeInBytes, ELogProgressUnits::MB, 1000, false /*bVerboseOnly*/);
+
+	FThreadLogConfig LogConfig;
+	auto BuildTargetCallback = [&Error, &PackDb, &DirCache, &OldDirectoryManifest, &Stats, &SyncProgressLogger, &LogConfig](
+								   const std::pair<const std::wstring, FFileManifest>& FileIt)
 	{
 		if (Error)
 		{
 			return;
 		}
 
+		FThreadLogConfig::FScope LogConfigScope(LogConfig);
+
 		const FFileManifest& FileManifest	= FileIt.second;
 		const FPath&		 TargetFilePath = FileManifest.CurrentPath;
 
+		auto				 OldManifestIt	 = OldDirectoryManifest.Files.find(FileIt.first);
+		const FFileManifest* OldFileManifest = nullptr;
+		if (OldManifestIt != OldDirectoryManifest.Files.end())
+		{
+			OldFileManifest = &OldManifestIt->second;
+
+			// TODO: option to diff blocks instead of just checking revision and metadata
+			// TODO: fix only the the read-only flag when that's the only difference
+			if (OldFileManifest->RevisionControlIdentity == FileManifest.RevisionControlIdentity &&
+				OldFileManifest->Mtime == FileManifest.Mtime && OldFileManifest->Size == FileManifest.Size &&
+				OldFileManifest->bReadOnly == FileManifest.bReadOnly)
+			{
+				UNSYNC_VERBOSE2(L"Skipped '%ls' (up to date)", FileIt.first.c_str());
+				Stats.NumFilesSkipped += 1;
+				Stats.BytesSkipped += FileManifest.Size;
+
+				SyncProgressLogger.Add(FileManifest.Size);
+
+				return;
+			}
+		}
+
 		FPath TargetFileParent = TargetFilePath.parent_path();
 		if (!DirCache.EnsureDirectoryExists(TargetFileParent))
 		{
@@ -855,6 +1003,8 @@ SyncDirectoryFromPack(const FPath& OutputRoot, const FPackDatabase& PackDb, cons
 			return;
 		}
 
+		// TODO: perform block diff before the expensive copy
+
 		if (GDryRun)
 		{
 			FNullReaderWriter NullFile(FileManifest.Size);
@@ -867,13 +1017,15 @@ SyncDirectoryFromPack(const FPath& OutputRoot, const FPackDatabase& PackDb, cons
 		}
 		else
 		{
-			// TODO: mark file writable if it exists
+			SetFileReadOnly(TargetFilePath, false);
+
 			FNativeFile TargetFile(TargetFilePath, EFileMode::CreateWriteOnly, FileManifest.Size);
 
 			if (!TargetFile.IsValid())
 			{
-				UNSYNC_ERROR(L"Failed to create target file '%ls'", TargetFilePath.wstring().c_str());
-				Error.Set(AppError("Failed to create target file"));
+				std::string SystemError = FormatSystemErrorMessage(TargetFile.GetError());
+				UNSYNC_ERROR(L"Failed to create target file '%ls'. %hs", TargetFilePath.wstring().c_str(), SystemError.c_str());
+				Error.Set(AppError("Failed to create target file", TargetFile.GetError()));
 				return;
 			}
 
@@ -893,10 +1045,18 @@ SyncDirectoryFromPack(const FPath& OutputRoot, const FPackDatabase& PackDb, cons
 				SetFileReadOnly(TargetFilePath, true);
 			}
 		}
-		
+
+		Stats.NumFilesSynced += 1;
+		Stats.BytesSynced += FileManifest.Size;
+
+		SyncProgressLogger.Add(FileManifest.Size);
 	};
 
 	ParallelForEach(NewDirectoryManifest.Files, BuildTargetCallback);
+	SyncProgressLogger.Complete();
+
+	UNSYNC_LOG(L"Skipped files: %llu (%.2f MB)", llu(Stats.NumFilesSkipped), SizeMb(Stats.BytesSkipped));
+	UNSYNC_LOG(L"Synced files: %llu (%.2f MB)", llu(Stats.NumFilesSynced), SizeMb(Stats.BytesSynced));
 
 	return !Error;
 }
@@ -1004,10 +1164,28 @@ CmdUnpack(const FCmdUnpackOptions& Options)
 				   double(PackDb.TotalCompressedSize) / double(1 << 30));
 	}
 
-	UNSYNC_LOG(L"Reading snapshot");
+	// Named/tagged snapshots take precendence over regular snapshots
+	FPath TagPath	   = Options.StorePath / "tag" / (Options.SnapshotName + ".unsync_tag");
+	FPath SnapshotPath = Options.StorePath / "snapshot" / (Options.SnapshotName + ".unsync_snapshot");
+
+	FBuffer SnapshotBuffer;
+
+	if (PathExists(TagPath))
+	{
+		UNSYNC_LOG(L"Loading snapshot '%hs' from tag file", Options.SnapshotName.c_str());
+		SnapshotBuffer = ReadFileToBuffer(TagPath);
+	}
+	else if (PathExists(SnapshotPath))
+	{
+		UNSYNC_LOG(L"Loading snapshot '%hs'", Options.SnapshotName.c_str());
+		SnapshotBuffer = ReadFileToBuffer(SnapshotPath);
+	}
+	else
+	{
+		UNSYNC_ERROR(L"Could not find snapshot file or named tag '%hs'", Options.SnapshotName.c_str());
+		return -1;
+	}
 
-	FPath	SnapshotPath   = Options.StorePath / (Options.SnapshotName + ".unsync_snapshot");
-	FBuffer SnapshotBuffer = ReadFileToBuffer(SnapshotPath);
 	if (SnapshotBuffer.Empty())
 	{
 		UNSYNC_ERROR(L"Failed read directory snapshot manifest");
@@ -1064,7 +1242,7 @@ CmdUnpack(const FCmdUnpackOptions& Options)
 
 	if (Options.bOutputFiles)
 	{
-		UNSYNC_LOG(L"Writing target files");
+		UNSYNC_LOG(L"Synchronizing directory from pack");
 		{
 			UNSYNC_LOG_INDENT;
 			if (!SyncDirectoryFromPack(Options.OutputPath, PackDb, NewDirectoryManifest))
@@ -1088,7 +1266,7 @@ CmdUnpack(const FCmdUnpackOptions& Options)
 	{
 		if (Options.bOutputRevisions)
 		{
-			UNSYNC_LOG(L"Extracting revision control identities");
+			UNSYNC_LOG(L"Extracting revision control data to file: '%ls'", RevisionFilePath.wstring().c_str());
 			UNSYNC_LOG_INDENT;
 
 			if (!SaveRevisionControlData(RevisionFilePath, NewDirectoryManifest, ERevisionControlFileFormat::IdentityOnly))
@@ -1100,7 +1278,7 @@ CmdUnpack(const FCmdUnpackOptions& Options)
 
 		if (!Options.P4HaveOutputPath.empty())
 		{
-			UNSYNC_LOG(L"Extracting revision control data into `p4 have` file '%ls'", Options.P4HaveOutputPath.wstring().c_str());
+			UNSYNC_LOG(L"Extracting revision control data to `p4 have` file '%ls'", Options.P4HaveOutputPath.wstring().c_str());
 			UNSYNC_LOG_INDENT;
 
 			if (!SaveRevisionControlData(Options.P4HaveOutputPath, NewDirectoryManifest, ERevisionControlFileFormat::P4Have))
diff --git a/Engine/Source/Programs/Unsync/Private/UnsyncCmdPack.h b/Engine/Source/Programs/Unsync/Private/UnsyncCmdPack.h
--- a/Engine/Source/Programs/Unsync/Private/UnsyncCmdPack.h
+++ b/Engine/Source/Programs/Unsync/Private/UnsyncCmdPack.h
@@ -15,6 +15,10 @@ struct FCmdPackOptions
 	FPath			  StorePath;   // optional
 	uint32			  BlockSize = uint32(64_KB);
 	FAlgorithmOptions Algorithm;
+
+	// Optional tag for this snapshot which can be used to reference it later using a human-friendly name.
+	// Overwrites the tag if it already exists.
+	std::string		  SnapshotName;
 };
 
 int32 CmdPack(const FCmdPackOptions& Options);
diff --git a/Engine/Source/Programs/Unsync/Private/UnsyncMain.cpp b/Engine/Source/Programs/Unsync/Private/UnsyncMain.cpp
--- a/Engine/Source/Programs/Unsync/Private/UnsyncMain.cpp
+++ b/Engine/Source/Programs/Unsync/Private/UnsyncMain.cpp
@@ -176,6 +176,7 @@ InnerMain(int Argc, char** Argv)
 		auto RunP4HaveOpt = SubPack->add_flag("--p4have", bRunP4Have, "Run `p4 have` when generating the dirctory pack");
 
 		SubPack->add_option("--store", StorePathUtf8, "Use this location to store pack data (default: <Input>/.unsync/pack)");
+		SubPack->add_option("--snapshot", SnapshotNameUtf8, "Custom name for the snapshot (will overwrite an existing tag)");
 
 		RunP4HaveOpt->excludes(P4HaveFileOpt);
 
@@ -815,6 +816,7 @@ InnerMain(int Argc, char** Argv)
 		PackOptions.bRunP4Have	 = bRunP4Have;
 		PackOptions.BlockSize	 = HashOrSyncBlockSize;
 		PackOptions.Algorithm	 = Algorithm;
+		PackOptions.SnapshotName = SnapshotNameUtf8;
 
 		return CmdPack(PackOptions);
 	}
@@ -1022,7 +1024,14 @@ ExceptionFilter(_EXCEPTION_POINTERS* ExceptionPointers)
 
 	PEXCEPTION_RECORD Record = ExceptionPointers->ExceptionRecord;
 
-	LogPrintf(ELogLevel::Error, L"Unhandled exception 0x%08X at address 0x%016X\n", Record->ExceptionCode, Record->ExceptionAddress);
+	if (Record->ExceptionCode == EXCEPTION_BREAKPOINT)
+	{
+		LogPrintf(ELogLevel::Error, L"Break point at address 0x%016X\n", Record->ExceptionAddress);
+	}
+	else
+	{
+		LogPrintf(ELogLevel::Error, L"Unhandled exception 0x%08X at address 0x%016X\n", Record->ExceptionCode, Record->ExceptionAddress);
+	}
 
 	LogWriteCrashDump(ExceptionPointers);
 
diff --git a/Engine/Source/Programs/Unsync/Private/UnsyncProgress.cpp b/Engine/Source/Programs/Unsync/Private/UnsyncProgress.cpp
--- a/Engine/Source/Programs/Unsync/Private/UnsyncProgress.cpp
+++ b/Engine/Source/Programs/Unsync/Private/UnsyncProgress.cpp
@@ -22,7 +22,7 @@ AddGlobalProgress(uint64 Size, EBlockListType ListType)
 	GGlobalProgressCurrent += Size;
 }
 
-FLogProgressScope::FLogProgressScope(uint64 InTotal, ELogProgressUnits InUnits, uint64 InPeriodMilliseconds)
+FLogProgressScope::FLogProgressScope(uint64 InTotal, ELogProgressUnits InUnits, uint64 InPeriodMilliseconds, bool bInVerboseOnly)
 : bParentThreadVerbose(GLogVerbose)
 , ParentThreadIndent(GLogIndent)
 , Current(0)
@@ -31,6 +31,7 @@ FLogProgressScope::FLogProgressScope(uint64 InTotal, ELogProgressUnits InUnits,
 , Units(InUnits)
 , NextProgressLogTime(TimePointNow())
 , bEnabled(true)
+, bVerboseOnly(bInVerboseOnly)
 {
 	Add(0);
 }
@@ -54,31 +55,34 @@ FLogProgressScope::Add(uint64 X, bool bForceComplete)
 	std::lock_guard<std::mutex> LockGuard(Mutex);
 	const uint64				CurrentClamped = std::min<uint64>(Current, Total);
 	const bool					bComplete	   = (CurrentClamped == Total) || bForceComplete;
-	if (bEnabled && GLogVerbose && (TimePointNow() > NextProgressLogTime || bComplete))
+
+	const ELogLevel LogLevel = bVerboseOnly ? ELogLevel::Debug : ELogLevel::Info;
+
+	if (bEnabled && (!bVerboseOnly || GLogVerbose) && (TimePointNow() > NextProgressLogTime || bComplete))
 	{
 		const wchar_t* Ending = bComplete ? L"\n" : L"\r";
 		switch (Units)
 		{
 			case ELogProgressUnits::GB:
-				LogPrintf(ELogLevel::Debug,
+				LogPrintf(LogLevel,
 						  L"%.2f / %.2f GB%ls",
 						  double(CurrentClamped) / double(1_GB),
 						  double(Total) / double(1_GB),
 						  Ending);
 				break;
 			case ELogProgressUnits::MB:
-				LogPrintf(ELogLevel::Debug,
+				LogPrintf(LogLevel,
 						  L"%.2f / %.2f MB%ls",
 						  double(CurrentClamped) / double(1_MB),
 						  double(Total) / double(1_MB),
 						  Ending);
 				break;
 			case ELogProgressUnits::Bytes:
-				LogPrintf(ELogLevel::Debug, L"%llu / %llu bytes%ls", (uint64)CurrentClamped, Total, Ending);
+				LogPrintf(LogLevel, L"%llu / %llu bytes%ls", (uint64)CurrentClamped, Total, Ending);
 				break;
 			case ELogProgressUnits::Raw:
 			default:
-				LogPrintf(ELogLevel::Debug, L"%llu / %llu%ls", (uint64)CurrentClamped, Total, Ending);
+				LogPrintf(LogLevel, L"%llu / %llu%ls", (uint64)CurrentClamped, Total, Ending);
 				break;
 		}
 
diff --git a/Engine/Source/Programs/Unsync/Private/UnsyncProgress.h b/Engine/Source/Programs/Unsync/Private/UnsyncProgress.h
--- a/Engine/Source/Programs/Unsync/Private/UnsyncProgress.h
+++ b/Engine/Source/Programs/Unsync/Private/UnsyncProgress.h
@@ -42,7 +42,7 @@ struct FLogProgressScope
 {
 	UNSYNC_DISALLOW_COPY_ASSIGN(FLogProgressScope)
 
-	FLogProgressScope(uint64 InTotal, ELogProgressUnits InUnits = ELogProgressUnits::Raw, uint64 InPeriodMilliseconds = 500);
+	FLogProgressScope(uint64 InTotal, ELogProgressUnits InUnits = ELogProgressUnits::Raw, uint64 InPeriodMilliseconds = 500, bool bInVerboseOnly = true);
 
 	void Complete();
 
@@ -58,6 +58,7 @@ struct FLogProgressScope
 	const ELogProgressUnits Units;
 	FTimePoint				NextProgressLogTime;
 	std::atomic<bool>		bEnabled;
+	const bool				bVerboseOnly;
 };
 
 }  // namespace unsync
diff --git a/Engine/Source/Programs/Unsync/Private/UnsyncUtil.cpp b/Engine/Source/Programs/Unsync/Private/UnsyncUtil.cpp
--- a/Engine/Source/Programs/Unsync/Private/UnsyncUtil.cpp
+++ b/Engine/Source/Programs/Unsync/Private/UnsyncUtil.cpp
@@ -90,20 +90,20 @@ void FTimingLogger::Finish()
 
 		if (Name.empty())
 		{
-			LogPrintf(LogLevel, L"%.3f sec", TotalSeconds);
+			LogPrintf(LogLevel, L"%.3f sec\n", TotalSeconds);
 		}
 		else
 		{
-			LogPrintf(LogLevel, L"%hs: %.3f sec", Name.c_str(), TotalSeconds);
-		}
-
-		if (TotalSeconds >= 60.0)
-		{
-			LogPrintf(LogLevel, L" (%02d:%02d:%02d)", H, M, S);
+			if (TotalSeconds >= 60.0)
+			{
+				LogPrintf(LogLevel, L"%hs: %.3f sec (%02d:%02d:%02d)\n", Name.c_str(), TotalSeconds, H, M, S);
+			}
+			else
+			{
+				LogPrintf(LogLevel, L"%hs: %.3f sec\n", Name.c_str(), TotalSeconds);
+			}
 		}
 
-		LogPrintf(LogLevel, L"\n");
-
 		LogFlush();
 
 		bEnabled = false;
