unsync - Split OpenID config out of FAuthDesc into a separate structure and use auth desc to generate token cache key instead of just server address

* Avoids unnecessary login steps when using different proxy servers that share auth config

diff --git a/Engine/Source/Programs/Unsync/Private/UnsyncAuth.cpp b/Engine/Source/Programs/Unsync/Private/UnsyncAuth.cpp
--- a/Engine/Source/Programs/Unsync/Private/UnsyncAuth.cpp
+++ b/Engine/Source/Programs/Unsync/Private/UnsyncAuth.cpp
@@ -270,14 +270,14 @@ DecodeJwtPayload(std::string JwtDataBase64Url)
 }
 
 TResult<FAuthToken>
-AcquireAuthToken(const FAuthDesc& AuthDesc)
+AcquireAuthToken(const FAuthDesc& AuthDesc, const FOpenIdConfig& OpenIdConfig)
 {
-	if (AuthDesc.AuthorizationEndpoint.empty())
+	if (OpenIdConfig.AuthorizationEndpoint.empty())
 	{
 		return AppError(L"Authorization endpoint is required");
 	}
 
-	if (AuthDesc.TokenEndpoint.empty())
+	if (OpenIdConfig.TokenEndpoint.empty())
 	{
 		return AppError(L"Token endpoint is required");
 	}
@@ -297,9 +297,15 @@ AcquireAuthToken(const FAuthDesc& AuthDesc)
 
 	FAuthToken Result;
 
-	FHttpConnection AuthServerConnection = FHttpConnection::CreateDefaultHttps(AuthDesc.ServerHost);
+	TResult<FRemoteDesc> AuthRemoteDesc = FRemoteDesc::FromUrl(AuthDesc.AuthServer);
+	if (AuthRemoteDesc.IsError())
+	{
+		return AppError(L"Failed to parse authentication server URI");
+	}
 
-	const uint16 CallbackPortNumber = CallbackServerDesc.HostPort;
+	FHttpConnection AuthServerConnection = FHttpConnection::CreateDefaultHttps(*AuthRemoteDesc);
+
+	const uint16 CallbackPortNumber = CallbackServerDesc.Host.Port;
 
 	FSocketHandle CallbackListenSocket = SocketListenTcp("127.0.0.1", CallbackPortNumber);
 
@@ -325,8 +331,8 @@ AcquireAuthToken(const FAuthDesc& AuthDesc)
 		"code_challenge={}&"
 		"state={}&"
 		"redirect_uri={}",
-		AuthDesc.ServerHost,
-		AuthDesc.AuthorizationEndpoint,
+		AuthRemoteDesc->Host.Address,
+		OpenIdConfig.AuthorizationEndpoint,
 		AuthDesc.ClientId,
 		AudienceParam,
 		CodeChallenge,
@@ -377,7 +383,7 @@ AcquireAuthToken(const FAuthDesc& AuthDesc)
 			CallbackUrl);
 
 		FHttpRequest Request;
-		Request.Url				   = AuthDesc.TokenEndpoint;
+		Request.Url				   = OpenIdConfig.TokenEndpoint;
 		Request.Method			   = EHttpMethod::POST;
 		Request.PayloadContentType = EHttpContentType::Application_WWWFormUrlEncoded;
 		Request.Payload			   = FBufferView{(const uint8*)TokenPayload.data(), (uint64)TokenPayload.size()};
@@ -433,15 +439,15 @@ AcquireAuthToken(const FAuthDesc& AuthDesc)
 }
 
 TResult<FAuthUserInfo>
-GetUserInfo(FHttpConnection& HttpConnection, const FAuthDesc& AuthDesc, const FAuthToken& AuthToken)
+GetUserInfo(FHttpConnection& HttpConnection, const FAuthDesc& AuthDesc, const FOpenIdConfig& OpenIdConfig, const FAuthToken& AuthToken)
 {
-	if (AuthDesc.UserInfoEndpoint.empty())
+	if (OpenIdConfig.UserInfoEndpoint.empty())
 	{
 		return AppError(L"User info endpoint is unknown");
 	}
 
 	FHttpRequest Request;
-	Request.Url			= AuthDesc.UserInfoEndpoint;
+	Request.Url			= OpenIdConfig.UserInfoEndpoint;
 	Request.Method		= EHttpMethod::GET;
 	Request.BearerToken = AuthToken.Access;
 
@@ -479,16 +485,22 @@ GetUserInfo(FHttpConnection& HttpConnection, const FAuthDesc& AuthDesc, const FA
 }
 
 TResult<FAuthToken>
-RefreshAuthToken(const FAuthDesc& AuthDesc, const FAuthToken& PreviousToken)
+RefreshAuthToken(const FAuthDesc& AuthDesc, const FOpenIdConfig& OpenIdConfig, const FAuthToken& PreviousToken)
 {
-	if (AuthDesc.TokenEndpoint.empty())
+	if (OpenIdConfig.TokenEndpoint.empty())
 	{
 		return AppError(L"Token endpoint is unknown");
 	}
 
 	FAuthToken Result = PreviousToken;
 
-	FHttpConnection AuthServerConnection = FHttpConnection::CreateDefaultHttps(AuthDesc.ServerHost);
+	TResult<FRemoteDesc> AuthRemoteDesc = FRemoteDesc::FromUrl(AuthDesc.AuthServer);
+	if (AuthRemoteDesc.IsError())
+	{
+		return AppError(L"Failed to parse authentication server URI");
+	}
+
+	FHttpConnection AuthServerConnection = FHttpConnection::CreateDefaultHttps(*AuthRemoteDesc);
 
 	std::string AccessToken;
 	std::string RefreshToken;
@@ -506,7 +518,7 @@ RefreshAuthToken(const FAuthDesc& AuthDesc, const FAuthToken& PreviousToken)
 			PreviousToken.Refresh);
 
 		FHttpRequest Request;
-		Request.Url				   = AuthDesc.TokenEndpoint;
+		Request.Url				   = OpenIdConfig.TokenEndpoint;
 		Request.Method			   = EHttpMethod::POST;
 		Request.PayloadContentType = EHttpContentType::Application_WWWFormUrlEncoded;
 		Request.Payload			   = FBufferView{(const uint8*)TokenPayload.data(), (uint64)TokenPayload.size()};
@@ -566,9 +578,15 @@ RefreshAuthToken(const FAuthDesc& AuthDesc, const FAuthToken& PreviousToken)
 }
 
 std::string
-GenerateTokenId(const FRemoteDesc& RemoteDesc)
+GenerateTokenId(const FAuthDesc& AuthDesc)
 {
-	FHash128 Hash = HashBlake3String<FHash128>(RemoteDesc.GetLoginAddress());
+	// TODO: just stream fields directly through a hasher
+	std::string HashInput;
+	HashInput += AuthDesc.AuthServer + " ";
+	HashInput += AuthDesc.ClientId + " ";
+	HashInput += AuthDesc.Audience;
+	//HashInput += AuthDesc.Callback; // don't need to consider the callback url
+	FHash128 Hash = HashBlake3String<FHash128>(HashInput);
 	return HashToHexString(Hash);
 }
 
@@ -690,12 +708,12 @@ LoadAuthToken(const FPath& Path)
 }
 
 TResult<FAuthToken>
-RefreshOrAcquireToken(const FAuthDesc& AuthDesc, const FAuthToken& PreviousToken)
+RefreshOrAcquireToken(const FAuthDesc& AuthDesc, const FOpenIdConfig& OpenIdConfig, const FAuthToken& PreviousToken)
 {
 	if (!PreviousToken.Refresh.empty())
 	{
 		UNSYNC_VERBOSE(L"Refreshing access token");
-		TResult<FAuthToken> RefreshResult = RefreshAuthToken(AuthDesc, PreviousToken);
+		TResult<FAuthToken> RefreshResult = RefreshAuthToken(AuthDesc, OpenIdConfig, PreviousToken);
 		if (RefreshResult.IsOk())
 		{
 			return RefreshResult;
@@ -703,11 +721,11 @@ RefreshOrAcquireToken(const FAuthDesc& AuthDesc, const FAuthToken& PreviousToken
 	}
 
 	UNSYNC_VERBOSE(L"Requesting new access token");
-	return AcquireAuthToken(AuthDesc);
+	return AcquireAuthToken(AuthDesc, OpenIdConfig);
 }
 
 TResult<FPath>
-GetTokenCachePath(const FRemoteDesc& RemoteDesc)
+GetTokenCachePath(const FAuthDesc& AuthDesc)
 {
 	FPath UserHomePath = GetUserHomeDirectory();
 	if (UserHomePath.empty())
@@ -715,7 +733,7 @@ GetTokenCachePath(const FRemoteDesc& RemoteDesc)
 		return AppError(L"Could not query user home directory path");
 	}
 
-	std::string TokenId = GenerateTokenId(RemoteDesc);
+	std::string TokenId = GenerateTokenId(AuthDesc);
 
 	FPath UnsyncSettingsPath = UserHomePath / FPath(".unsync");
 	FPath TokenCachePath	 = UnsyncSettingsPath / FPath(TokenId);
@@ -741,8 +759,25 @@ LogAuthTokenExpiration(const FAuthToken& AuthToken)
 	}
 }
 
+FAuthDesc
+FAuthDesc::FromHelloResponse(const ProxyQuery::FHelloResponse& HelloResponse)
+{
+	FAuthDesc AuthDesc;
+	AuthDesc.AuthServer = HelloResponse.AuthServerUri;
+	AuthDesc.ClientId	= HelloResponse.AuthClientId;
+	AuthDesc.Audience	= HelloResponse.AuthAudience;
+	AuthDesc.Callback	= HelloResponse.CallbackUri;
+
+	if (AuthDesc.Callback.empty())
+	{
+		AuthDesc.Callback = "http://localhost:8080";  // sensible default
+	}
+
+	return AuthDesc;
+}
+
 TResult<FAuthToken>
-Authenticate(const FRemoteDesc& RemoteDesc, int32 RefreshThreshold)
+Authenticate(const FAuthDesc& AuthDesc, int32 RefreshThreshold)
 {
 	// Authentication must be serialized (only one thread should ever open the browser for interactive login, etc.)
 	static std::mutex			AuthMutex;
@@ -750,7 +785,7 @@ Authenticate(const FRemoteDesc& RemoteDesc, int32 RefreshThreshold)
 
 	FAuthToken PreviousToken;
 
-	TResult<FPath> TokenCachePathResult = GetTokenCachePath(RemoteDesc);
+	TResult<FPath> TokenCachePathResult = GetTokenCachePath(AuthDesc);
 
 	if (const FPath* TokenCachePath = TokenCachePathResult.TryData())
 	{
@@ -787,15 +822,13 @@ Authenticate(const FRemoteDesc& RemoteDesc, int32 RefreshThreshold)
 		return ResultOk(PreviousToken);
 	}
 
-	TResult<FAuthDesc> AuthDescResult = GetAuthenticationDesc(RemoteDesc);
-	if (AuthDescResult.IsError())
+	TResult<FOpenIdConfig> OpenIdConfigResult = GetOpenIdConfig(AuthDesc);
+	if (OpenIdConfigResult.IsError())
 	{
-		return MoveError<FAuthToken>(AuthDescResult);
+		return MoveError<FAuthToken>(OpenIdConfigResult);
 	}
 
-	const FAuthDesc& AuthDesc = AuthDescResult.GetData();
-
-	TResult<FAuthToken> FreshTokenResult = RefreshOrAcquireToken(AuthDesc, PreviousToken);
+	TResult<FAuthToken> FreshTokenResult = RefreshOrAcquireToken(AuthDesc, *OpenIdConfigResult, PreviousToken);
 	if (FreshTokenResult.IsError())
 	{
 		return FreshTokenResult;
@@ -821,79 +854,49 @@ Authenticate(const FRemoteDesc& RemoteDesc, int32 RefreshThreshold)
 	return FreshTokenResult;
 }
 
-bool
-TryAddAuthentication(FRemoteDesc& InOutRemoteDesc)
-{
-	TResult<FAuthToken> AuthTokenResult = Authenticate(InOutRemoteDesc, 5 * 60);
-
-	if (AuthTokenResult.IsError())
-	{
-		return false;
-	}
-
-	// Authentication requires encrypted connection
-	InOutRemoteDesc.bTlsEnable				= true;
-	InOutRemoteDesc.bAuthenticationRequired = true;
-
-	return true;
-}
-
 TResult<FAuthDesc>
-GetAuthenticationDesc(const FRemoteDesc& RemoteDesc)
+GetRemoteAuthDesc(const FRemoteDesc& RemoteDesc)
 {
-	// TODO: possibly other backend could use automatic authentication also
-	if (RemoteDesc.Protocol != EProtocolFlavor::Unsync)
-	{
-		return AppError(L"Authentication is only implemented for UNSYNC protocol");
-	}
-
-	FRemoteDesc LoginRemoteDesc = RemoteDesc;
-	LoginRemoteDesc.HostAddress = RemoteDesc.GetLoginAddress();
-
-	TResult<ProxyQuery::FHelloResponse> HelloResponseResult = ProxyQuery::Hello(LoginRemoteDesc, /*bAnonymous*/ true);
+	TResult<ProxyQuery::FHelloResponse> HelloResponseResult = ProxyQuery::Hello(RemoteDesc, nullptr /*AuthDesc: null for anonymous initial connection*/);
 	if (HelloResponseResult.IsError())
 	{
+		UNSYNC_ERROR("Failed establish a handshake with server '%hs'", RemoteDesc.Host.Address.c_str());
+		LogError(HelloResponseResult.GetError());
 		return MoveError<FAuthDesc>(HelloResponseResult);
 	}
 
-	if (!HelloResponseResult->SupportsAuthentication())
-	{
-		return AppError(L"Server does not support authentication");
-	}
+	const FAuthDesc AuthDesc = FAuthDesc::FromHelloResponse(*HelloResponseResult);
 
-	FAuthDesc AuthDesc;
-	AuthDesc.ClientId = HelloResponseResult->AuthClientId;
-	AuthDesc.Audience = HelloResponseResult->AuthAudience;
-	AuthDesc.Callback = HelloResponseResult->CallbackUri;
-	if (AuthDesc.Callback.empty())
+	return ResultOk(AuthDesc);
+}
+
+TResult<FOpenIdConfig>
+GetOpenIdConfig(const FAuthDesc& AuthDesc)
+{
+	if (!AuthDesc.IsValid())
 	{
-		AuthDesc.Callback = "http://localhost:8080";  // sensible default
+		return AppError(L"Mandatory authentication parameters not provided");
 	}
 
-	TResult<FRemoteDesc> AuthServerDescResult = FRemoteDesc::FromUrl(HelloResponseResult->AuthServerUri);
+	TResult<FRemoteDesc> AuthServerDescResult = FRemoteDesc::FromUrl(AuthDesc.AuthServer);
 	if (AuthServerDescResult.IsError())
 	{
-		return MoveError<FAuthDesc>(AuthServerDescResult);
+		return MoveError<FOpenIdConfig>(AuthServerDescResult);
 	}
 
 	const FRemoteDesc& AuthServerDesc = AuthServerDescResult.GetData();
 
-	AuthDesc.ServerHost = AuthServerDesc.HostAddress;
-
 	std::string ServerApiPrefix;
 	if (!AuthServerDesc.RequestPath.empty())
 	{
 		ServerApiPrefix = fmt::format("/{}", AuthServerDesc.RequestPath);
 	}
 
-	if (!AuthDesc.IsValid())
-	{
-		return AppError(L"Mandatory authentication parameters not found");
-	}
-
-	FHttpConnection AuthServerConnection = FHttpConnection::CreateDefaultHttps(AuthDesc.ServerHost);
+	FHttpConnection AuthServerConnection = FHttpConnection::CreateDefaultHttps(*AuthServerDescResult);
 	std::string		ConfigEndpoint		 = fmt::format("{}/.well-known/openid-configuration", ServerApiPrefix);
 
+	FOpenIdConfig OpenIdConfig;
+
 	FHttpResponse ConfigResponse = HttpRequest(AuthServerConnection, EHttpMethod::GET, ConfigEndpoint);
 	if (ConfigResponse.Success())
 	{
@@ -903,7 +906,7 @@ GetAuthenticationDesc(const FRemoteDesc& RemoteDesc)
 		std::string JsonErrorString;
 		Json		JsonObject = Json::parse(JsonString, JsonErrorString);
 
-		std::string EndpointPrefix = fmt::format("https://{}", AuthDesc.ServerHost);
+		std::string EndpointPrefix = fmt::format("https://{}", AuthServerDescResult->Host.Address);
 
 		if (JsonErrorString.empty())
 		{
@@ -919,14 +922,14 @@ GetAuthenticationDesc(const FRemoteDesc& RemoteDesc)
 				return {};
 			};
 
-			AuthDesc.AuthorizationEndpoint = ExtractEndpoint("authorization_endpoint");
-			AuthDesc.TokenEndpoint		   = ExtractEndpoint("token_endpoint");
-			AuthDesc.UserInfoEndpoint	   = ExtractEndpoint("userinfo_endpoint");
-			AuthDesc.JwksUri			   = JsonObject["jwks_uri"].string_value();
+			OpenIdConfig.AuthorizationEndpoint = ExtractEndpoint("authorization_endpoint");
+			OpenIdConfig.TokenEndpoint		   = ExtractEndpoint("token_endpoint");
+			OpenIdConfig.UserInfoEndpoint	   = ExtractEndpoint("userinfo_endpoint");
+			OpenIdConfig.JwksUri			   = JsonObject["jwks_uri"].string_value();
 		}
 	}
 
-	return ResultOk(AuthDesc);
+	return ResultOk(OpenIdConfig);
 }
 
 int64
diff --git a/Engine/Source/Programs/Unsync/Private/UnsyncAuth.h b/Engine/Source/Programs/Unsync/Private/UnsyncAuth.h
--- a/Engine/Source/Programs/Unsync/Private/UnsyncAuth.h
+++ b/Engine/Source/Programs/Unsync/Private/UnsyncAuth.h
@@ -14,6 +14,10 @@ namespace unsync {
 struct FHttpConnection;
 struct FRemoteDesc;
 
+namespace ProxyQuery {
+	struct FHelloResponse;
+}
+
 struct FAuthToken
 {
 	// Raw data acquired from authentication endpoint
@@ -25,18 +29,22 @@ struct FAuthToken
 	int64		ExirationTime = 0; // UNIX timestamp in seconds (0 if unknown)
 };
 
-struct FAuthDesc
+// Data from /.well-known/openid-configuration
+struct FOpenIdConfig
 {
-	// Mandatory configuration
-	std::string ServerHost; // only the hostname itself
-	std::string ClientId;
-	
-	// Mandatory fields for certain APIs
 	std::string AuthorizationEndpoint;
 	std::string TokenEndpoint;
 	std::string UserInfoEndpoint;
 	std::string JwksUri;
+};
 
+// Data from server configuration
+struct FAuthDesc
+{
+	// Mandatory configuration
+	std::string AuthServer;
+	std::string ClientId;
+	
 	// Optional configuration
 	std::string Audience;
 	std::string Callback;	// OAuth flow redirection URL
@@ -44,25 +52,27 @@ struct FAuthDesc
 	// Check that mandatory fields have values
 	bool IsValid() const
 	{
-		return !ServerHost.empty() && !ClientId.empty();
+		return !AuthServer.empty() && !ClientId.empty();
 	}
+
+	static FAuthDesc FromHelloResponse(const ProxyQuery::FHelloResponse& HelloResponse);
 };
 
-// Query OIDC/OAuth2 configuration using server handshake
-TResult<FAuthDesc> GetAuthenticationDesc(const FRemoteDesc& RemoteDesc);
+// Query OIDC/OAuth2 configuration via HTTP
+TResult<FOpenIdConfig> GetOpenIdConfig(const FAuthDesc& AuthDesc);
 
 // Perform the complete authentication flow:
 // - Attempt use stored refresh token first
 // - If refresh is not possible, use PKCE Authentication flow to get new tokens
 // - Save refresh token in user directory for future use
 // - Skips acquiring new token if remaining valid time is above RefreshThreshold (in seconds)
-TResult<FAuthToken> Authenticate(const FRemoteDesc& RemoteDesc, int32 RefreshThreshold = INT_MAX);
+TResult<FAuthToken> Authenticate(const FAuthDesc& AuthDesc, int32 RefreshThreshold = INT_MAX);
 
 // Auth utility functions
 
-TResult<FAuthToken> AcquireAuthToken(const FAuthDesc& AuthDesc);
-TResult<FAuthToken> RefreshAuthToken(const FAuthDesc& AuthDesc, const FAuthToken& PreviousToken);
-TResult<FAuthToken> RefreshOrAcquireToken(const FAuthDesc& AuthDesc, const FAuthToken& PreviousToken);
+TResult<FAuthToken> AcquireAuthToken(const FAuthDesc& AuthDesc, const FOpenIdConfig& OpenIdConfig);
+TResult<FAuthToken> RefreshAuthToken(const FAuthDesc& AuthDesc, const FOpenIdConfig& OpenIdConfig, const FAuthToken& PreviousToken);
+TResult<FAuthToken> RefreshOrAcquireToken(const FAuthDesc& AuthDesc, const FOpenIdConfig& OpenIdConfig, const FAuthToken& PreviousToken);
 
 struct FAuthUserInfo
 {
@@ -90,7 +100,7 @@ void TransformBase64UrlSafeToVanilla(std::string& Base64UrlSafe);
 
 std::string GetPKCECodeChallenge(std::string_view CodeVerifier);
 
-bool TryAddAuthentication(FRemoteDesc& InOutRemoteDesc);
+TResult<FAuthDesc> GetRemoteAuthDesc(const FRemoteDesc& RemoteDesc);
 
 int64 GetSecondsFromUnixEpoch();
 
diff --git a/Engine/Source/Programs/Unsync/Private/UnsyncCmdLogin.cpp b/Engine/Source/Programs/Unsync/Private/UnsyncCmdLogin.cpp
--- a/Engine/Source/Programs/Unsync/Private/UnsyncCmdLogin.cpp
+++ b/Engine/Source/Programs/Unsync/Private/UnsyncCmdLogin.cpp
@@ -35,7 +35,18 @@ CmdLogin(const FCmdLoginOptions& Options)
 {
 	const int32 RefreshThreshold = Options.bForceRefresh ? INT_MAX : 5 * 60;
 
-	TResult<FAuthToken> AuthTokenResult = Authenticate(Options.Remote, RefreshThreshold);
+	UNSYNC_VERBOSE(L"Connecting to '%hs'", Options.Remote.Host.Address.c_str());
+	TResult<ProxyQuery::FHelloResponse> HelloResponseResult = ProxyQuery::Hello(Options.Remote, nullptr /*AuthDesc: anonymous initial connection*/);
+	if (HelloResponseResult.IsError())
+	{
+		UNSYNC_ERROR("Failed establish a handshake with server '%hs'", Options.Remote.Host.Address.c_str());
+		LogError(HelloResponseResult.GetError());
+		return -1;
+	}
+
+	const FAuthDesc AuthDesc = FAuthDesc::FromHelloResponse(*HelloResponseResult);
+
+	TResult<FAuthToken> AuthTokenResult = Authenticate(AuthDesc, RefreshThreshold);
 
 	if (AuthTokenResult.IsOk())
 	{
diff --git a/Engine/Source/Programs/Unsync/Private/UnsyncCmdQuery.cpp b/Engine/Source/Programs/Unsync/Private/UnsyncCmdQuery.cpp
--- a/Engine/Source/Programs/Unsync/Private/UnsyncCmdQuery.cpp
+++ b/Engine/Source/Programs/Unsync/Private/UnsyncCmdQuery.cpp
@@ -144,7 +144,19 @@ CmdQueryMirrors(const FCmdQueryOptions& Options)
 int32
 CmdQueryList(const FCmdQueryOptions& Options)
 {
-	TResult<FAuthToken> AuthToken = Authenticate(Options.Remote, 5 * 60);
+	FHttpConnection Connection = FHttpConnection::CreateDefaultHttps(Options.Remote);
+
+	TResult<ProxyQuery::FHelloResponse> HelloResponse = ProxyQuery::Hello(Connection);
+	if (HelloResponse.IsError())
+	{
+		UNSYNC_ERROR("Failed establish a handshake with server '%hs'", Options.Remote.Host.Address.c_str());
+		LogError(HelloResponse.GetError());
+		return -1;
+	}
+	FAuthDesc AuthDesc = FAuthDesc::FromHelloResponse(*HelloResponse);
+
+	TResult<FAuthToken> AuthToken = Authenticate(AuthDesc, 5 * 60);
+
 	if (!AuthToken.IsOk())
 	{
 		LogError(AuthToken.GetError());
@@ -157,7 +169,6 @@ CmdQueryList(const FCmdQueryOptions& Options)
 		return -1;
 	}
 
-	FHttpConnection Connection = FHttpConnection::CreateDefaultHttps(Options.Remote);
 
 	std::string Url = fmt::format("/api/v1/list?{}", Options.Args[0]);
 
@@ -188,13 +199,6 @@ CmdQuerySearch(const FCmdQueryOptions& Options)
 {
 	using namespace ProxyQuery;
 
-	TResult<FAuthToken> AuthToken = Authenticate(Options.Remote, 5 * 60);
-	if (!AuthToken.IsOk())
-	{
-		LogError(AuthToken.GetError());
-		return -1;
-	}
-
 	if (Options.Args.empty())
 	{
 		UNSYNC_ERROR(L"Path argument is required");
@@ -204,11 +208,36 @@ CmdQuerySearch(const FCmdQueryOptions& Options)
 	auto CreateConnection = [Remote = Options.Remote]
 	{
 		FTlsClientSettings TlsSettings = Remote.GetTlsClientSettings();
-		return new FHttpConnection(Remote.HostAddress, Remote.HostPort, &TlsSettings);
+		return new FHttpConnection(Remote.Host.Address, Remote.Host.Port, &TlsSettings);
 	};
 
 	TObjectPool<FHttpConnection> ConnectionPool(CreateConnection);
 
+	std::string BearerToken;
+
+	{
+		std::unique_ptr<FHttpConnection> Connection = ConnectionPool.Acquire();
+		TResult<ProxyQuery::FHelloResponse> HelloResponse = ProxyQuery::Hello(*Connection);
+		ConnectionPool.Release(std::move(Connection));
+		if (HelloResponse.IsError())
+		{
+			UNSYNC_ERROR("Failed establish a handshake with server '%hs'", Options.Remote.Host.Address.c_str());
+			LogError(HelloResponse.GetError());
+			return -1;
+		}
+
+		FAuthDesc AuthDesc = FAuthDesc::FromHelloResponse(*HelloResponse);
+
+		TResult<FAuthToken> AuthToken = Authenticate(AuthDesc, 5 * 60);
+		if (!AuthToken.IsOk())
+		{
+			LogError(AuthToken.GetError());
+			return -1;
+		}
+
+		BearerToken = std::move(AuthToken->Access);
+	}
+
 	const std::string& RootPath = Options.Args[0];
 	std::vector<std::regex> SubdirPatterns;
 
@@ -261,7 +290,7 @@ CmdQuerySearch(const FCmdQueryOptions& Options)
 	Context.ParentThreadIndent	 = GLogIndent;
 
 	std::function<void(std::string, int32)> ExploreDirectory =
-		[&Context, &AuthToken, &ConnectionPool, &ExploreDirectory, &SubdirPatterns, &Tasks](std::string Path, int32 CurrentDepth)
+		[&Context, &BearerToken, &ConnectionPool, &ExploreDirectory, &SubdirPatterns, &Tasks](std::string Path, int32 CurrentDepth)
 	{
 		FLogVerbosityScope VerboseScope(Context.bParentThreadVerbose);
 		FLogIndentScope	   IndentScope(Context.ParentThreadIndent, true);
@@ -273,7 +302,7 @@ CmdQuerySearch(const FCmdQueryOptions& Options)
 		FHttpRequest Request;
 		Request.Url			= Url;
 		Request.Method		= EHttpMethod::GET;
-		Request.BearerToken = AuthToken->Access;
+		Request.BearerToken = BearerToken;
 
 		Context.ConnectionSemaphore.Acquire();
 
@@ -370,6 +399,26 @@ CmdQueryFile(const FCmdQueryOptions& Options)
 		return -1;
 	}
 
+	// TODO: use a global connection pool and use it for ProxyQuery::DownloadFile too
+	FHttpConnection HelloConnection = FHttpConnection::CreateDefaultHttps(Options.Remote);
+	TResult<ProxyQuery::FHelloResponse> HelloResponse = ProxyQuery::Hello(HelloConnection);
+	if (HelloResponse.IsError())
+	{
+		UNSYNC_ERROR("Failed establish a handshake with server '%hs'", Options.Remote.Host.Address.c_str());
+		LogError(HelloResponse.GetError());
+		return -1;
+	}
+	FAuthDesc AuthDesc = FAuthDesc::FromHelloResponse(*HelloResponse);
+	HelloConnection.Close();
+
+	TResult<FAuthToken> AuthToken = Authenticate(AuthDesc, 5 * 60);
+
+	if (!AuthToken.IsOk())
+	{
+		LogError(AuthToken.GetError());
+		return -1;
+	}
+
 	UNSYNC_LOG(L"Downloading file: '%hs'", Options.Args[0].c_str());
 	UNSYNC_LOG_INDENT;
 
@@ -402,7 +451,7 @@ CmdQueryFile(const FCmdQueryOptions& Options)
 		return *ResultWriter;
 	};
 
-	TResult<> Response = ProxyQuery::DownloadFile(Options.Remote, Options.Args[0], OutputCallback);
+	TResult<> Response = ProxyQuery::DownloadFile(Options.Remote, &AuthDesc, Options.Args[0], OutputCallback);
 
 	if (Response.IsOk())
 	{
diff --git a/Engine/Source/Programs/Unsync/Private/UnsyncCmdSync.cpp b/Engine/Source/Programs/Unsync/Private/UnsyncCmdSync.cpp
--- a/Engine/Source/Programs/Unsync/Private/UnsyncCmdSync.cpp
+++ b/Engine/Source/Programs/Unsync/Private/UnsyncCmdSync.cpp
@@ -21,7 +21,7 @@ CmdSync(const FCmdSyncOptions& Options)
 		return Options.Filter ? Options.Filter->Resolve(Path) : Path;
 	};
 
-	FProxyPool ProxyPool(Options.Remote);
+	FProxyPool ProxyPool(Options.Remote, Options.AuthDesc);
 
 	std::error_code ErrorCode	   = {};
 	FPath			ResolvedSource = ResolvePath(Options.Source);
diff --git a/Engine/Source/Programs/Unsync/Private/UnsyncCmdSync.h b/Engine/Source/Programs/Unsync/Private/UnsyncCmdSync.h
--- a/Engine/Source/Programs/Unsync/Private/UnsyncCmdSync.h
+++ b/Engine/Source/Programs/Unsync/Private/UnsyncCmdSync.h
@@ -8,6 +8,8 @@
 
 namespace unsync {
 
+struct FAuthDesc;
+
 struct FCmdSyncOptions
 {
 	FAlgorithmOptions Algorithm;
@@ -22,6 +24,7 @@ struct FCmdSyncOptions
 	std::vector<FPath> Overlays;
 
 	FRemoteDesc Remote;
+	FAuthDesc*	AuthDesc = nullptr;
 
 	bool bFullDifference = false;
 	bool bFullSourceScan = false;
diff --git a/Engine/Source/Programs/Unsync/Private/UnsyncCore.cpp b/Engine/Source/Programs/Unsync/Private/UnsyncCore.cpp
--- a/Engine/Source/Programs/Unsync/Private/UnsyncCore.cpp
+++ b/Engine/Source/Programs/Unsync/Private/UnsyncCore.cpp
@@ -27,7 +27,7 @@ UNSYNC_THIRD_PARTY_INCLUDES_START
 #include <md5-sse2.h>
 UNSYNC_THIRD_PARTY_INCLUDES_END
 
-#define UNSYNC_VERSION_STR "1.0.63-dev1"
+#define UNSYNC_VERSION_STR "1.0.63-dev2"
 
 namespace unsync {
 
@@ -1849,7 +1849,7 @@ struct FPooledProxy
 };
 
 static bool
-DownloadFileIfNewer(const FRemoteDesc& RemoteDesc, const FPath& Source, const FPath& Target, EFileMode TargetFileMode)
+DownloadFileIfNewer(const FRemoteDesc& RemoteDesc, const FAuthDesc* AuthDesc, const FPath& Source, const FPath& Target, EFileMode TargetFileMode)
 {
 	using FDirectoryListing		 = ProxyQuery::FDirectoryListing;
 	using FDirectoryListingEntry = ProxyQuery::FDirectoryListingEntry;
@@ -1862,7 +1862,7 @@ DownloadFileIfNewer(const FRemoteDesc& RemoteDesc, const FPath& Source, const FP
 	std::string SourceFileNameUtf8 = ConvertWideToUtf8(SourceFileName.wstring());
 
 	// TODO: could have a dedicated single file stat query
-	TResult<FDirectoryListing> DirectoryListingResult = ProxyQuery::ListDirectory(RemoteDesc, SourceParentUtf8);
+	TResult<FDirectoryListing> DirectoryListingResult = ProxyQuery::ListDirectory(RemoteDesc, AuthDesc, SourceParentUtf8);
 	if (DirectoryListingResult.IsError())
 	{
 		LogError(DirectoryListingResult.GetError());
@@ -1890,7 +1890,7 @@ DownloadFileIfNewer(const FRemoteDesc& RemoteDesc, const FPath& Source, const FP
 	if (SourceEntry->Size != TargetAttr.Size || SourceEntry->Mtime != TargetAttr.Mtime)
 	{
 		UNSYNC_VERBOSE(L"Downloading '%ls'", Source.wstring().c_str());
-		TResult<FBuffer> DownloadResult = ProxyQuery::DownloadFile(RemoteDesc, SourceUtf8);
+		TResult<FBuffer> DownloadResult = ProxyQuery::DownloadFile(RemoteDesc, AuthDesc, SourceUtf8);
 		if (DownloadResult.IsError())
 		{
 			LogError(DownloadResult.GetError());
@@ -1968,6 +1968,7 @@ LoadAndMergeSourceManifest(FDirectoryManifest& Output,
 	{
 		UNSYNC_LOG_INDENT;
 		bool bDownloadedOk = DownloadFileIfNewer(ProxyPool.RemoteDesc,
+												 ProxyPool.AuthDesc,
 												 SourceManifestPath,
 												 SourceManifestTempPath,
 												 EFileMode::CreateWriteOnly | EFileMode::IgnoreDryRun);
@@ -2633,8 +2634,8 @@ SyncDirectory(const FSyncDirectoryOptions& SyncOptions)
 		LogGlobalStatus(L"Connecting to server");
 		UNSYNC_LOG(L"Connecting to %hs server '%hs:%d' ...",
 					   ToString(ProxyPool.RemoteDesc.Protocol),
-					   ProxyPool.RemoteDesc.HostAddress.c_str(),
-					   ProxyPool.RemoteDesc.HostPort);
+					   ProxyPool.RemoteDesc.Host.Address.c_str(),
+					   ProxyPool.RemoteDesc.Host.Port);
 		UNSYNC_LOG_INDENT;
 
 		std::unique_ptr<FProxy> Proxy = ProxyPool.Alloc();
diff --git a/Engine/Source/Programs/Unsync/Private/UnsyncHash.cpp b/Engine/Source/Programs/Unsync/Private/UnsyncHash.cpp
--- a/Engine/Source/Programs/Unsync/Private/UnsyncHash.cpp
+++ b/Engine/Source/Programs/Unsync/Private/UnsyncHash.cpp
@@ -65,21 +65,6 @@ HashBytes(const uint8* Data, uint64 Size, EHashType HashType)
 	return Result;
 }
 
-FHash128
-compute_hash_128(const uint8* Data, uint64 Size, EStrongHashAlgorithmID Algorithm)
-{
-	switch (Algorithm)
-	{
-		case EStrongHashAlgorithmID::MD5:
-			return HashMd5Bytes(Data, Size);
-		case EStrongHashAlgorithmID::Blake3_128:
-			return HashBlake3Bytes<FHash128>(Data, Size);
-		default:
-			UNSYNC_ERROR(L"Unexpected hash algorithm id");
-			return FHash128{};
-	}
-}
-
 FGenericHash
 ComputeHash(const uint8* Data, uint64 Size, EStrongHashAlgorithmID Algorithm)
 {
diff --git a/Engine/Source/Programs/Unsync/Private/UnsyncHttp.cpp b/Engine/Source/Programs/Unsync/Private/UnsyncHttp.cpp
--- a/Engine/Source/Programs/Unsync/Private/UnsyncHttp.cpp
+++ b/Engine/Source/Programs/Unsync/Private/UnsyncHttp.cpp
@@ -259,7 +259,7 @@ HttpRequest(const FRemoteDesc& RemoteDesc,
 			std::string_view   BearerToken)
 {
 	FTlsClientSettings TlsSettings = RemoteDesc.GetTlsClientSettings();
-	FHttpConnection	   Connection(RemoteDesc.HostAddress, RemoteDesc.HostPort, RemoteDesc.bTlsEnable ? &TlsSettings : nullptr);
+	FHttpConnection	   Connection(RemoteDesc.Host.Address, RemoteDesc.Host.Port, RemoteDesc.bTlsEnable ? &TlsSettings : nullptr);
 
 	FHttpRequest Request;
 
@@ -540,7 +540,7 @@ FHttpConnection
 FHttpConnection::CreateDefaultHttps(const FRemoteDesc& RemoteDesc)
 {
 	FTlsClientSettings TlsSettings = RemoteDesc.GetTlsClientSettings();
-	return FHttpConnection(RemoteDesc.HostAddress, RemoteDesc.HostPort, &TlsSettings);
+	return FHttpConnection(RemoteDesc.Host.Address, RemoteDesc.Host.Port, &TlsSettings);
 }
 
 bool
diff --git a/Engine/Source/Programs/Unsync/Private/UnsyncJupiter.cpp b/Engine/Source/Programs/Unsync/Private/UnsyncJupiter.cpp
--- a/Engine/Source/Programs/Unsync/Private/UnsyncJupiter.cpp
+++ b/Engine/Source/Programs/Unsync/Private/UnsyncJupiter.cpp
@@ -467,14 +467,14 @@ TResult<uint64>
 JupiterPush(const FDirectoryManifest& Manifest, const FRemoteDesc& RemoteDesc, FTlsClientSettings* TlsSettings)
 {
 	auto CreateConnection = [RemoteDesc, TlsSettings] {
-		return new FHttpConnection(RemoteDesc.HostAddress, RemoteDesc.HostPort, TlsSettings);
+		return new FHttpConnection(RemoteDesc.Host.Address, RemoteDesc.Host.Port, TlsSettings);
 	};
 
 	FSemaphore					 ChunkUploadSemaphore(8);  // up to 8 concurrent connections
 	TObjectPool<FHttpConnection> ConnectionPool(CreateConnection);
 
 	{
-		UNSYNC_VERBOSE(L"Connecting to Jupiter server '%hs:%d'", RemoteDesc.HostAddress.c_str(), RemoteDesc.HostPort);
+		UNSYNC_VERBOSE(L"Connecting to Jupiter server '%hs:%d'", RemoteDesc.Host.Address.c_str(), RemoteDesc.Host.Port);
 
 		// TODO: add an RAII helper to return connections to the pool
 		std::unique_ptr<FHttpConnection> Connection = ConnectionPool.Acquire();
@@ -700,7 +700,7 @@ FJupiterProtocolImpl::FJupiterProtocolImpl(const FRemoteDesc&		 InSettings,
 										   const FTlsClientSettings* TlsSettings,
 										   std::string_view			 InHttpHeaders)
 : FRemoteProtocolBase(InSettings, InRequestMap)
-, Connection(InSettings.HostAddress, InSettings.HostPort, InSettings.bTlsEnable ? TlsSettings : nullptr)
+, Connection(InSettings.Host.Address, InSettings.Host.Port, InSettings.bTlsEnable ? TlsSettings : nullptr)
 , HttpHeaders(InHttpHeaders)
 , RemoteDesc(InSettings)
 {
diff --git a/Engine/Source/Programs/Unsync/Private/UnsyncMain.cpp b/Engine/Source/Programs/Unsync/Private/UnsyncMain.cpp
--- a/Engine/Source/Programs/Unsync/Private/UnsyncMain.cpp
+++ b/Engine/Source/Programs/Unsync/Private/UnsyncMain.cpp
@@ -532,6 +532,7 @@ InnerMain(int Argc, char** Argv)
 	}
 
 	FRemoteDesc RemoteDesc;
+	FAuthDesc	AuthDesc;
 
 	if (RemoteAddressUtf8.empty())
 	{
@@ -767,7 +768,7 @@ InnerMain(int Argc, char** Argv)
 
 	if (bShouldLogin)
 	{
-		RemoteDesc.LoginAddress = RemoteDesc.HostAddress;
+		RemoteDesc.PrimaryHost = RemoteDesc.Host;
 	}
 
 	FRemoteDesc RootRemoteDesc = RemoteDesc;
@@ -776,14 +777,14 @@ InnerMain(int Argc, char** Argv)
 		&& Cli.got_subcommand(SubSync)
 		&& RemoteDesc.IsValid() && RemoteDesc.Protocol == EProtocolFlavor::Unsync)
 	{
-		UNSYNC_LOG(L"Selecting server using root '%hs'", RemoteDesc.HostAddress.c_str());
+		UNSYNC_LOG(L"Selecting server using root '%hs'", RemoteDesc.Host.Address.c_str());
 		TResult<FMirrorInfo> MirrorResult = FindClosestMirror(RemoteDesc);
 		if (const FMirrorInfo* Mirror = MirrorResult.TryData())
 		{
 			UNSYNC_LOG(L"Closest server: '%hs', ping: %.2f ms", Mirror->Address.c_str(), Mirror->Ping * 1000.0);
 
-			RemoteDesc.HostAddress = Mirror->Address;
-			RemoteDesc.HostPort	   = Mirror->Port;
+			RemoteDesc.Host.Address = Mirror->Address;
+			RemoteDesc.Host.Port	= Mirror->Port;
 		}
 		else
 		{
@@ -859,10 +860,29 @@ InnerMain(int Argc, char** Argv)
 		{
 			UNSYNC_LOG(L"Attempting to authenticate");
 			UNSYNC_LOG_INDENT;
-			bool bUsingAuthentication = TryAddAuthentication(RemoteDesc);
-			if (bUsingAuthentication)
+
+			TResult<FAuthDesc> AuthDescResult = GetRemoteAuthDesc(RemoteDesc);
+
+			if (AuthDescResult.IsOk())
 			{
-				UNSYNC_LOG(L"Authentication enabled");
+				AuthDesc = AuthDescResult.GetData();
+
+				// Note: since tokens can expire during a long operation,
+				// we can only save the auth descriptor and re-authenticate later if necessary
+				TResult<FAuthToken> AuthTokenResult = Authenticate(AuthDesc, 5 * 60);
+
+				if (AuthTokenResult.IsError())
+				{
+					UNSYNC_ERROR("Failed to authenticate with server '%hs'", RemoteDesc.Host.Address.c_str());
+					LogError(AuthTokenResult.GetError());
+					return -1;
+				}
+
+				// Authentication requires encrypted connection
+				RemoteDesc.bTlsEnable			   = true;
+				RemoteDesc.bAuthenticationRequired = true;
+
+				UNSYNC_LOG(L"Authentication enabled")
 			}
 		}
 
@@ -902,6 +922,7 @@ InnerMain(int Argc, char** Argv)
 		SyncOptions.Target					   = TargetFilename;
 		SyncOptions.SourceManifestOverride	   = SourceManifestFilename;
 		SyncOptions.Remote					   = RemoteDesc;
+		SyncOptions.AuthDesc				   = AuthDesc.IsValid() ? &AuthDesc : nullptr;
 		SyncOptions.bFullDifference			   = bFullDifference;
 		SyncOptions.bFullSourceScan			   = bFullSourceScan;
 		SyncOptions.bCleanup				   = !bNoCleanupAfterSync;
diff --git a/Engine/Source/Programs/Unsync/Private/UnsyncProxy.cpp b/Engine/Source/Programs/Unsync/Private/UnsyncProxy.cpp
--- a/Engine/Source/Programs/Unsync/Private/UnsyncProxy.cpp
+++ b/Engine/Source/Programs/Unsync/Private/UnsyncProxy.cpp
@@ -18,6 +18,7 @@ struct FUnsyncProtocolImpl : FRemoteProtocolBase
 {
 	FUnsyncProtocolImpl(const FRemoteDesc&			   InRemoteDesc,
 						const FRemoteProtocolFeatures& InFeatures,
+						const FAuthDesc*			   InAuthDesc,
 						const FBlockRequestMap*		   InRequestMap,
 						const FTlsClientSettings*	   TlsSettings);
 	virtual ~FUnsyncProtocolImpl() override;
@@ -40,7 +41,7 @@ struct FUnsyncProtocolImpl : FRemoteProtocolBase
 	static void SendTelemetryEvent(const FRemoteDesc& RemoteDesc, const FTelemetryEventSyncComplete& Event);
 };
 
-FProxy::FProxy(const FRemoteDesc& RemoteDesc, const FRemoteProtocolFeatures& InFeatures, const FBlockRequestMap* InRequestMap)
+FProxy::FProxy(const FRemoteDesc& RemoteDesc, const FRemoteProtocolFeatures& InFeatures, const FAuthDesc* InAuthDesc, const FBlockRequestMap* InRequestMap)
 {
 	UNSYNC_ASSERT(InRequestMap);
 
@@ -48,12 +49,13 @@ FProxy::FProxy(const FRemoteDesc& RemoteDesc, const FRemoteProtocolFeatures& InF
 
 	if (RemoteDesc.Protocol == EProtocolFlavor::Jupiter)
 	{
-		ProtocolImpl =
-			std::unique_ptr<FRemoteProtocolBase>(new FJupiterProtocolImpl(RemoteDesc, InRequestMap, &TlsSettings, RemoteDesc.HttpHeaders));
+		auto Inner	 = new FJupiterProtocolImpl(RemoteDesc, InRequestMap, &TlsSettings, RemoteDesc.HttpHeaders);
+		ProtocolImpl = std::unique_ptr<FRemoteProtocolBase>(Inner);
 	}
 	else if (RemoteDesc.Protocol == EProtocolFlavor::Unsync)
 	{
-		ProtocolImpl = std::unique_ptr<FRemoteProtocolBase>(new FUnsyncProtocolImpl(RemoteDesc, InFeatures, InRequestMap, &TlsSettings));
+		auto* Inner	 = new FUnsyncProtocolImpl(RemoteDesc, InFeatures, InAuthDesc, InRequestMap, &TlsSettings);
+		ProtocolImpl = std::unique_ptr<FRemoteProtocolBase>(Inner);
 	}
 	else
 	{
@@ -63,6 +65,7 @@ FProxy::FProxy(const FRemoteDesc& RemoteDesc, const FRemoteProtocolFeatures& InF
 
 FUnsyncProtocolImpl::FUnsyncProtocolImpl(const FRemoteDesc&				RemoteDesc,
 										 const FRemoteProtocolFeatures& InFeatures,
+										 const FAuthDesc*				InAuthDesc,
 										 const FBlockRequestMap*		InRequestMap,
 										 const FTlsClientSettings*		TlsSettings)
 : FRemoteProtocolBase(RemoteDesc, InRequestMap)
@@ -70,7 +73,7 @@ FUnsyncProtocolImpl::FUnsyncProtocolImpl(const FRemoteDesc&				RemoteDesc,
 {
 	if (RemoteDesc.bTlsEnable && TlsSettings)
 	{
-		FSocketHandle RawSocketHandle = SocketConnectTcp(RemoteDesc.HostAddress.c_str(), RemoteDesc.HostPort);
+		FSocketHandle RawSocketHandle = SocketConnectTcp(RemoteDesc.Host.Address.c_str(), RemoteDesc.Host.Port);
 		SocketSetRecvTimeout(RawSocketHandle, RemoteDesc.RecvTimeoutSeconds);
 
 		if (RawSocketHandle)
@@ -89,7 +92,7 @@ FUnsyncProtocolImpl::FUnsyncProtocolImpl(const FRemoteDesc&				RemoteDesc,
 
 	if (!SocketHandle)
 	{
-		FSocketHandle RawSocketHandle = SocketConnectTcp(RemoteDesc.HostAddress.c_str(), RemoteDesc.HostPort);
+		FSocketHandle RawSocketHandle = SocketConnectTcp(RemoteDesc.Host.Address.c_str(), RemoteDesc.Host.Port);
 		SocketSetRecvTimeout(RawSocketHandle, RemoteDesc.RecvTimeoutSeconds);
 
 		SocketHandle = std::unique_ptr<FSocketRaw>(new FSocketRaw(RawSocketHandle));
@@ -127,35 +130,43 @@ FUnsyncProtocolImpl::FUnsyncProtocolImpl(const FRemoteDesc&				RemoteDesc,
 
 	if (IsValid() && Features.bAuthentication && RemoteDesc.bAuthenticationRequired)
 	{
-		bool bOk = IsValid();
+		if (InAuthDesc)
+		{
+			bool bOk = IsValid();
 
-		TResult<FAuthToken> AuthTokenResult = Authenticate(RemoteDesc, 15 * 60);
+			TResult<FAuthToken> AuthTokenResult = Authenticate(*InAuthDesc, 15 * 60);
 
-		if (AuthTokenResult.IsOk())
-		{
-			FBufferView AccessToken = {(const uint8*)AuthTokenResult->Access.data(), AuthTokenResult->Access.length()};
+			if (AuthTokenResult.IsOk())
+			{
+				FBufferView AccessToken = {(const uint8*)AuthTokenResult->Access.data(), AuthTokenResult->Access.length()};
+
+				FCommandPacket Packet;
+				Packet.CommandId = COMMAND_ID_AUTHENTICATE;
+				bOk &= SendStruct(*SocketHandle, Packet);
+				bOk &= SendBuffer(*SocketHandle, AccessToken);
 
-			FCommandPacket Packet;
-			Packet.CommandId = COMMAND_ID_AUTHENTICATE;
-			bOk &= SendStruct(*SocketHandle, Packet);
-			bOk &= SendBuffer(*SocketHandle, AccessToken);
+				int32 ResultSize = 0;
+				bOk &= SocketRecvT(*SocketHandle, ResultSize);
 
-			int32 ResultSize = 0;
-			bOk &= SocketRecvT(*SocketHandle, ResultSize);
+				FBuffer ResultBuffer;
+				if (ResultSize)
+				{
+					ResultBuffer.Resize(ResultSize);
+					bOk &= (SocketRecvAll(*SocketHandle, ResultBuffer.Data(), ResultSize) == ResultSize);
+				}
 
-			FBuffer ResultBuffer;
-			if (ResultSize)
+				// TODO: parse authentication result packet and report errors
+			}
+			else
 			{
-				ResultBuffer.Resize(ResultSize);
-				bOk &= (SocketRecvAll(*SocketHandle, ResultBuffer.Data(), ResultSize) == ResultSize);
+				LogError(AuthTokenResult.GetError());
+				UNSYNC_ERROR(L"Server requires authentication, but access token could not be acquired");
+				Invalidate();
 			}
-
-			// TODO: parse authentication result packet and report errors
 		}
 		else
 		{
-			LogError(AuthTokenResult.GetError());
-			UNSYNC_ERROR(L"Server requires authentication");
+			UNSYNC_ERROR(L"Server requires authentication, but required parameters were not provided");
 			Invalidate();
 		}
 	}
@@ -394,26 +405,34 @@ FUnsyncProtocolImpl::GetSocketSecurity() const
 }
 
 TResult<ProxyQuery::FHelloResponse>
-ProxyQuery::Hello(const FRemoteDesc& RemoteDesc, bool bAnonymous)
+ProxyQuery::Hello(const FRemoteDesc& RemoteDesc, const FAuthDesc* OptAuthDesc)
+{
+	FTlsClientSettings TlsSettings = RemoteDesc.GetTlsClientSettings();
+	FHttpConnection	   Connection(RemoteDesc.Host.Address, RemoteDesc.Host.Port, RemoteDesc.bTlsEnable ? &TlsSettings : nullptr);
+	return Hello(Connection, OptAuthDesc);
+}
+
+TResult<ProxyQuery::FHelloResponse>
+ProxyQuery::Hello(FHttpConnection& HttpConnection, const FAuthDesc* OptAuthDesc)
 {
 	const char* Url = "/api/v1/hello";
 
 	std::string BearerToken;
-	if (RemoteDesc.bAuthenticationRequired && !bAnonymous)
+	if (OptAuthDesc)
 	{
-		TResult<FAuthToken> AuthTokenResult = Authenticate(RemoteDesc, 15 * 60);
+		TResult<FAuthToken> AuthTokenResult = Authenticate(*OptAuthDesc, 15 * 60);
 		if (AuthTokenResult.IsOk())
 		{
 			BearerToken = std::move(AuthTokenResult.GetData().Access);
 		}
 	}
 
-	FHttpResponse Response = HttpRequest(RemoteDesc, EHttpMethod::GET, Url, BearerToken);
+	FHttpResponse Response = HttpRequest(HttpConnection, EHttpMethod::GET, Url, {} /*CustomHeaders*/, BearerToken);
 
 	if (!Response.Success())
 	{
 		UNSYNC_ERROR(L"Failed to establish connection to UNSYNC server. Error code: %d.", Response.Code);
-		return HttpError(fmt::format("{}:{}{}", RemoteDesc.HostAddress.c_str(), RemoteDesc.HostPort, Url), Response.Code);
+		return HttpError(fmt::format("{}:{}{}", HttpConnection.HostAddress.c_str(), HttpConnection.HostPort, Url), Response.Code);
 	}
 
 	FHelloResponse Result;
@@ -499,6 +518,16 @@ ProxyQuery::Hello(const FRemoteDesc& RemoteDesc, bool bAnonymous)
 		}
 	}
 
+	if (auto& Field = JsonObject["primary"]; Field.is_string())
+	{
+		const std::string& PrimaryHostStr = Field.string_value();
+		TResult<FRemoteDesc> PrimaryHostDesc = FRemoteDesc::FromUrl(PrimaryHostStr);
+		if (PrimaryHostDesc.IsOk())
+		{
+			Result.PrimaryHost = PrimaryHostDesc->Host;
+		}
+	}
+
 	return ResultOk(std::move(Result));
 }
 
@@ -549,22 +578,28 @@ ProxyQuery::FDirectoryListing::FromJson(const char* JsonString)
 }
 
 TResult<ProxyQuery::FDirectoryListing>
-ProxyQuery::ListDirectory(const FRemoteDesc& Remote, const std::string& Path)
+ProxyQuery::ListDirectory(const FRemoteDesc& Remote, const FAuthDesc* AuthDesc, const std::string& Path)
 {
-	TResult<FAuthToken> AuthToken = Authenticate(Remote, 5 * 60);
-	if (!AuthToken.IsOk())
+	std::string Url = fmt::format("/api/v1/list?{}", Path);
+
+	std::string BearerToken;
+	if (AuthDesc)
 	{
-		return MoveError<FDirectoryListing>(AuthToken);
+		TResult<FAuthToken> AuthToken = Authenticate(*AuthDesc, 5 * 60);
+		if (!AuthToken.IsOk())
+		{
+			return MoveError<FDirectoryListing>(AuthToken);
+		}
+
+		BearerToken = std::move(AuthToken->Access);
 	}
 
 	FHttpConnection Connection = FHttpConnection::CreateDefaultHttps(Remote);
 
-	std::string Url = fmt::format("/api/v1/list?{}", Path);
-
 	FHttpRequest Request;
 	Request.Url			= Url;
 	Request.Method		= EHttpMethod::GET;
-	Request.BearerToken = AuthToken->Access;
+	Request.BearerToken = BearerToken;
 
 	FHttpResponse Response = HttpRequest(Connection, Request);
 
@@ -579,12 +614,15 @@ ProxyQuery::ListDirectory(const FRemoteDesc& Remote, const std::string& Path)
 }
 
 TResult<>
-ProxyQuery::DownloadFile(const FRemoteDesc& Remote, const std::string& Path, ProxyQuery::FDownloadOutputCallback OutputCallback)
+ProxyQuery::DownloadFile(const FRemoteDesc&					 Remote,
+						 const FAuthDesc*					 AuthDesc,
+						 const std::string&					 Path,
+						 ProxyQuery::FDownloadOutputCallback OutputCallback)
 {
 	auto CreateConnection = [Remote]
 	{
 		FTlsClientSettings TlsSettings = Remote.GetTlsClientSettings();
-		return new FHttpConnection(Remote.HostAddress, Remote.HostPort, &TlsSettings);
+		return new FHttpConnection(Remote.Host.Address, Remote.Host.Port, &TlsSettings);
 	};
 
 	TObjectPool<FHttpConnection> ConnectionPool(CreateConnection);
@@ -596,16 +634,21 @@ ProxyQuery::DownloadFile(const FRemoteDesc& Remote, const std::string& Path, Pro
 	{
 		std::unique_ptr<FHttpConnection> Connection = ConnectionPool.Acquire();
 
-		TResult<FAuthToken> AuthToken = Authenticate(Remote, 5 * 60);
-		if (!AuthToken.IsOk())
+		std::string BearerToken;
+		if (AuthDesc)
 		{
-			return std::move(AuthToken.GetError());
+			TResult<FAuthToken> AuthToken = Authenticate(*AuthDesc, 5 * 60);
+			if (!AuthToken.IsOk())
+			{
+				return std::move(AuthToken.GetError());
+			}
+			BearerToken = std::move(AuthToken->Access);
 		}
 
 		FHttpRequest HeadRequest;
 		HeadRequest.Url			   = Url;
 		HeadRequest.Method		   = EHttpMethod::HEAD;
-		HeadRequest.BearerToken	   = AuthToken->Access;
+		HeadRequest.BearerToken	   = BearerToken;
 		FHttpResponse HeadResponse = HttpRequest(*Connection, HeadRequest);
 		if (!HeadResponse.Success())
 		{
@@ -643,8 +686,20 @@ ProxyQuery::DownloadFile(const FRemoteDesc& Remote, const std::string& Path, Pro
 
 	FLogProgressScope DownloadProgress(FileSize, ELogProgressUnits::MB);
 
+	std::string BearerToken;
+	if (AuthDesc)
+	{
+		TResult<FAuthToken> AuthToken = Authenticate(*AuthDesc, 5 * 60);
+		if (!AuthToken.IsOk())
+		{
+			Error.Set(std::move(AuthToken.GetError()));
+			return AppError(L"Failed to acquire access token");
+		}
+		BearerToken = std::move(AuthToken->Access);
+	}
+
 	auto ProcessChunk =
-		[&Error, &Result, &Url, &Remote, &ConnectionPool, &DownloadSempahore, &DownloadProgress](
+		[&Error, &Result, &Url, &Remote, &ConnectionPool, &DownloadSempahore, &DownloadProgress, &BearerToken](
 			const FRange& Range)
 	{
 		FLogIndentScope	   IndentScope(DownloadProgress.ParentThreadIndent, true);
@@ -657,13 +712,6 @@ ProxyQuery::DownloadFile(const FRemoteDesc& Remote, const std::string& Path, Pro
 
 		DownloadSempahore.Acquire();
 
-		TResult<FAuthToken> AuthToken = Authenticate(Remote, 5 * 60);
-		if (!AuthToken.IsOk())
-		{
-			Error.Set(std::move(AuthToken.GetError()));
-			return;
-		}
-
 		std::unique_ptr<FHttpConnection> Connection = ConnectionPool.Acquire();
 
 		std::string RequestHeaders = fmt::format("Range: bytes={}-{}", Range.Offset, Range.Offset + Range.Size - 1);
@@ -671,7 +719,7 @@ ProxyQuery::DownloadFile(const FRemoteDesc& Remote, const std::string& Path, Pro
 		FHttpRequest Request;
 		Request.Url			   = Url;
 		Request.Method		   = EHttpMethod::GET;
-		Request.BearerToken	   = AuthToken->Access;
+		Request.BearerToken	   = BearerToken;
 		Request.CustomHeaders  = RequestHeaders;
 		FHttpResponse Response = HttpRequest(*Connection, Request);
 
@@ -708,7 +756,7 @@ ProxyQuery::DownloadFile(const FRemoteDesc& Remote, const std::string& Path, Pro
 }
 
 TResult<FBuffer>
-ProxyQuery::DownloadFile(const FRemoteDesc& Remote, const std::string& Path)
+ProxyQuery::DownloadFile(const FRemoteDesc& Remote, const FAuthDesc* AuthDesc, const std::string& Path)
 {
 	FBuffer Result;
 
@@ -721,7 +769,7 @@ ProxyQuery::DownloadFile(const FRemoteDesc& Remote, const std::string& Path)
 		return *ResultWriter;
 	};
 
-	TResult<> DownloadResult = DownloadFile(Remote, Path, OutputCallback);
+	TResult<> DownloadResult = DownloadFile(Remote, AuthDesc, Path, OutputCallback);
 	if (DownloadResult.IsOk())
 	{
 		return ResultOk(std::move(Result));
@@ -889,13 +937,14 @@ FBlockRequestMap::GetMacroBlockRequest(const FGenericHash& BlockHash) const
 	return Result;
 }
 
-FProxyPool::FProxyPool() : FProxyPool(FRemoteDesc())
+FProxyPool::FProxyPool() : FProxyPool(FRemoteDesc(), nullptr)
 {
 }
 
-FProxyPool::FProxyPool(const FRemoteDesc& InRemoteDesc)
+FProxyPool::FProxyPool(const FRemoteDesc& InRemoteDesc, const FAuthDesc* InAuthDesc)
 : ParallelDownloadSemaphore(InRemoteDesc.MaxConnections)
 , RemoteDesc(InRemoteDesc)
+, AuthDesc(InAuthDesc)
 , bValid(InRemoteDesc.IsValid())
 {
 	if (!bValid)
@@ -907,10 +956,10 @@ FProxyPool::FProxyPool(const FRemoteDesc& InRemoteDesc)
 	{
 		UNSYNC_VERBOSE(L"Connecting to %hs server '%hs:%d' ...",
 					   ToString(RemoteDesc.Protocol),
-					   RemoteDesc.HostAddress.c_str(),
-					   RemoteDesc.HostPort);
+					   RemoteDesc.Host.Address.c_str(),
+					   RemoteDesc.Host.Port);
 
-		TResult<ProxyQuery::FHelloResponse> Response = ProxyQuery::Hello(RemoteDesc);
+		TResult<ProxyQuery::FHelloResponse> Response = ProxyQuery::Hello(RemoteDesc, AuthDesc);
 
 		if (Response.IsError())
 		{
@@ -955,7 +1004,7 @@ FProxyPool::Alloc()
 
 	if (!Result || !Result->IsValid())
 	{
-		Result = std::make_unique<FProxy>(RemoteDesc, Features, &RequestMap);
+ 		Result = std::make_unique<FProxy>(RemoteDesc, Features, AuthDesc, &RequestMap);
 	}
 
 	return Result;
diff --git a/Engine/Source/Programs/Unsync/Private/UnsyncProxy.h b/Engine/Source/Programs/Unsync/Private/UnsyncProxy.h
--- a/Engine/Source/Programs/Unsync/Private/UnsyncProxy.h
+++ b/Engine/Source/Programs/Unsync/Private/UnsyncProxy.h
@@ -12,10 +12,13 @@
 #include <functional>
 #include <mutex>
 #include <unordered_map>
+#include <optional>
 
 namespace unsync {
 
 struct FDirectoryManifest;
+struct FAuthDesc;
+struct FHttpConnection;
 
 enum class EDownloadRetryMode
 {
@@ -136,7 +139,10 @@ struct FRemoteProtocolBase
 class FProxy
 {
 public:
-	FProxy(const FRemoteDesc& InRemoteDesc, const FRemoteProtocolFeatures& InFeatures, const FBlockRequestMap* InRequestMap);
+	FProxy(const FRemoteDesc&			  InRemoteDesc,
+		   const FRemoteProtocolFeatures& InFeatures,
+		   const FAuthDesc*				  InAuthDesc,
+		   const FBlockRequestMap*		  InRequestMap);
 	~FProxy();
 
 	bool Contains(const FDirectoryManifest& Manifest);
@@ -153,7 +159,7 @@ class FProxyPool
 {
 public:
 	FProxyPool();
-	FProxyPool(const FRemoteDesc& InRemoteDesc);
+	FProxyPool(const FRemoteDesc& InRemoteDesc, const FAuthDesc* InAuthDesc);
 
 	std::unique_ptr<FProxy> Alloc();
 	void					Dealloc(std::unique_ptr<FProxy>&& Proxy);
@@ -164,6 +170,7 @@ public:
 	FSemaphore ParallelDownloadSemaphore;
 
 	const FRemoteDesc RemoteDesc;
+	const FAuthDesc* AuthDesc = nullptr; // optional reference to externally-owned auth parameters
 
 	void InitRequestMap(EStrongHashAlgorithmID InStrongHasher);
 	void BuildFileBlockRequests(const FPath& OriginalFilePath, const FPath& ResolvedFilePath, const FFileManifest& FileManifest);
@@ -202,9 +209,12 @@ struct FHelloResponse
 	std::vector<std::string> FeatureNames;
 	FRemoteProtocolFeatures	 Features;
 
+	std::optional<FHostAddressAndPort> PrimaryHost;
+
 	bool SupportsAuthentication() const { return Features.bAuthentication && !AuthServerUri.empty() && !AuthClientId.empty(); }
 };
-TResult<FHelloResponse> Hello(const FRemoteDesc& RemoteDesc, bool bAnonymous = false);
+TResult<FHelloResponse> Hello(const FRemoteDesc& RemoteDesc, const FAuthDesc* OptAuthDesc = nullptr);
+TResult<FHelloResponse> Hello(FHttpConnection& Connection, const FAuthDesc* OptAuthDesc = nullptr);
 
 struct FDirectoryListingEntry
 {
@@ -221,11 +231,12 @@ struct FDirectoryListing
 	static TResult<FDirectoryListing> FromJson(const char* JsonString);
 };
 
-TResult<FDirectoryListing> ListDirectory(const FRemoteDesc& Remote, const std::string& Path);
-TResult<FBuffer>		   DownloadFile(const FRemoteDesc& Remote, const std::string& Path);
+// TODO: add overloads with HTTP connection
+TResult<FDirectoryListing> ListDirectory(const FRemoteDesc& Remote, const FAuthDesc* AuthDesc, const std::string& Path);
+TResult<FBuffer>		   DownloadFile(const FRemoteDesc& Remote, const FAuthDesc* AuthDesc, const std::string& Path);
 
 using FDownloadOutputCallback = std::function<FIOWriter&(uint64 Size)>;
-TResult<> DownloadFile(const FRemoteDesc& Remote, const std::string& Path, FDownloadOutputCallback OutputCallback);
+TResult<> DownloadFile(const FRemoteDesc& Remote, const FAuthDesc* AuthDesc, const std::string& Path, FDownloadOutputCallback OutputCallback);
 
 } 
 
diff --git a/Engine/Source/Programs/Unsync/Private/UnsyncRemote.cpp b/Engine/Source/Programs/Unsync/Private/UnsyncRemote.cpp
--- a/Engine/Source/Programs/Unsync/Private/UnsyncRemote.cpp
+++ b/Engine/Source/Programs/Unsync/Private/UnsyncRemote.cpp
@@ -168,8 +168,8 @@ FRemoteDesc::FromUrl(std::string_view Url)
 		return AppError("Invalid host port");  // TODO: extract the port substring
 	}
 
-	Result.HostAddress = HostAddress;
-	Result.HostPort	   = HostPort;
+	Result.Host.Address = HostAddress;
+	Result.Host.Port	= HostPort;
 
 	Result.bTlsEnable			 = bUseTls;
 	Result.bTlsVerifyCertificate = bUseTls;
@@ -211,8 +211,8 @@ TestParseRemote()
 		TResult<FRemoteDesc> ParseResult = FRemoteDesc::FromUrl("example.com");
 		UNSYNC_ASSERT(ParseResult.IsOk());
 		UNSYNC_ASSERT(ParseResult->bTlsEnable == false);
-		UNSYNC_ASSERT(ParseResult->HostAddress == "example.com");
-		UNSYNC_ASSERT(ParseResult->HostPort == UNSYNC_DEFAULT_PORT);
+		UNSYNC_ASSERT(ParseResult->Host.Address == "example.com");
+		UNSYNC_ASSERT(ParseResult->Host.Port == UNSYNC_DEFAULT_PORT);
 		UNSYNC_ASSERT(ParseResult->Protocol == EProtocolFlavor::Unsync);
 	}
 
@@ -220,8 +220,8 @@ TestParseRemote()
 		TResult<FRemoteDesc> ParseResult = FRemoteDesc::FromUrl("http://example.com");
 		UNSYNC_ASSERT(ParseResult.IsOk());
 		UNSYNC_ASSERT(ParseResult->bTlsEnable == false);
-		UNSYNC_ASSERT(ParseResult->HostAddress == "example.com");
-		UNSYNC_ASSERT(ParseResult->HostPort == 80);
+		UNSYNC_ASSERT(ParseResult->Host.Address == "example.com");
+		UNSYNC_ASSERT(ParseResult->Host.Port == 80);
 		UNSYNC_ASSERT(ParseResult->Protocol == EProtocolFlavor::Jupiter);
 	}
 
@@ -229,8 +229,8 @@ TestParseRemote()
 		TResult<FRemoteDesc> ParseResult = FRemoteDesc::FromUrl("https://example.com");
 		UNSYNC_ASSERT(ParseResult.IsOk());
 		UNSYNC_ASSERT(ParseResult->bTlsEnable == true);
-		UNSYNC_ASSERT(ParseResult->HostAddress == "example.com");
-		UNSYNC_ASSERT(ParseResult->HostPort == 443);
+		UNSYNC_ASSERT(ParseResult->Host.Address == "example.com");
+		UNSYNC_ASSERT(ParseResult->Host.Port == 443);
 		UNSYNC_ASSERT(ParseResult->Protocol == EProtocolFlavor::Jupiter);
 	}
 
@@ -238,8 +238,8 @@ TestParseRemote()
 		TResult<FRemoteDesc> ParseResult = FRemoteDesc::FromUrl("http://example.com:1234");
 		UNSYNC_ASSERT(ParseResult.IsOk());
 		UNSYNC_ASSERT(ParseResult->bTlsEnable == false);
-		UNSYNC_ASSERT(ParseResult->HostAddress == "example.com");
-		UNSYNC_ASSERT(ParseResult->HostPort == 1234);
+		UNSYNC_ASSERT(ParseResult->Host.Address == "example.com");
+		UNSYNC_ASSERT(ParseResult->Host.Port == 1234);
 		UNSYNC_ASSERT(ParseResult->Protocol == EProtocolFlavor::Jupiter);
 	}
 
@@ -247,8 +247,8 @@ TestParseRemote()
 		TResult<FRemoteDesc> ParseResult = FRemoteDesc::FromUrl("https://example.com:1234");
 		UNSYNC_ASSERT(ParseResult.IsOk());
 		UNSYNC_ASSERT(ParseResult->bTlsEnable == true);
-		UNSYNC_ASSERT(ParseResult->HostAddress == "example.com");
-		UNSYNC_ASSERT(ParseResult->HostPort == 1234);
+		UNSYNC_ASSERT(ParseResult->Host.Address == "example.com");
+		UNSYNC_ASSERT(ParseResult->Host.Port == 1234);
 		UNSYNC_ASSERT(ParseResult->Protocol == EProtocolFlavor::Jupiter);
 	}
 
@@ -256,8 +256,8 @@ TestParseRemote()
 		TResult<FRemoteDesc> ParseResult = FRemoteDesc::FromUrl("jupiter+http://example.com:1234");
 		UNSYNC_ASSERT(ParseResult.IsOk());
 		UNSYNC_ASSERT(ParseResult->bTlsEnable == false);
-		UNSYNC_ASSERT(ParseResult->HostAddress == "example.com");
-		UNSYNC_ASSERT(ParseResult->HostPort == 1234);
+		UNSYNC_ASSERT(ParseResult->Host.Address == "example.com");
+		UNSYNC_ASSERT(ParseResult->Host.Port == 1234);
 		UNSYNC_ASSERT(ParseResult->Protocol == EProtocolFlavor::Jupiter);
 	}
 
@@ -265,8 +265,8 @@ TestParseRemote()
 		TResult<FRemoteDesc> ParseResult = FRemoteDesc::FromUrl("jupiter+https://example.com");
 		UNSYNC_ASSERT(ParseResult.IsOk());
 		UNSYNC_ASSERT(ParseResult->bTlsEnable == true);
-		UNSYNC_ASSERT(ParseResult->HostAddress == "example.com");
-		UNSYNC_ASSERT(ParseResult->HostPort == 443);
+		UNSYNC_ASSERT(ParseResult->Host.Address == "example.com");
+		UNSYNC_ASSERT(ParseResult->Host.Port == 443);
 		UNSYNC_ASSERT(ParseResult->Protocol == EProtocolFlavor::Jupiter);
 	}
 
@@ -274,8 +274,8 @@ TestParseRemote()
 		TResult<FRemoteDesc> ParseResult = FRemoteDesc::FromUrl("jupiter+https://example.com#test.namespace");
 		UNSYNC_ASSERT(ParseResult.IsOk());
 		UNSYNC_ASSERT(ParseResult->bTlsEnable == true);
-		UNSYNC_ASSERT(ParseResult->HostAddress == "example.com");
-		UNSYNC_ASSERT(ParseResult->HostPort == 443);
+		UNSYNC_ASSERT(ParseResult->Host.Address == "example.com");
+		UNSYNC_ASSERT(ParseResult->Host.Port == 443);
 		UNSYNC_ASSERT(ParseResult->Protocol == EProtocolFlavor::Jupiter);
 		UNSYNC_ASSERT(ParseResult->StorageNamespace == "test.namespace");
 	}
@@ -284,8 +284,8 @@ TestParseRemote()
 		TResult<FRemoteDesc> ParseResult = FRemoteDesc::FromUrl("jupiter+https://example.com:1234#test.namespace");
 		UNSYNC_ASSERT(ParseResult.IsOk());
 		UNSYNC_ASSERT(ParseResult->bTlsEnable == true);
-		UNSYNC_ASSERT(ParseResult->HostAddress == "example.com");
-		UNSYNC_ASSERT(ParseResult->HostPort == 1234);
+		UNSYNC_ASSERT(ParseResult->Host.Address == "example.com");
+		UNSYNC_ASSERT(ParseResult->Host.Port == 1234);
 		UNSYNC_ASSERT(ParseResult->Protocol == EProtocolFlavor::Jupiter);
 		UNSYNC_ASSERT(ParseResult->StorageNamespace == "test.namespace");
 	}
@@ -294,8 +294,8 @@ TestParseRemote()
 		TResult<FRemoteDesc> ParseResult = FRemoteDesc::FromUrl("example.com:1234#test.namespace");
 		UNSYNC_ASSERT(ParseResult.IsOk());
 		UNSYNC_ASSERT(ParseResult->bTlsEnable == false);
-		UNSYNC_ASSERT(ParseResult->HostAddress == "example.com");
-		UNSYNC_ASSERT(ParseResult->HostPort == 1234);
+		UNSYNC_ASSERT(ParseResult->Host.Address == "example.com");
+		UNSYNC_ASSERT(ParseResult->Host.Port == 1234);
 		UNSYNC_ASSERT(ParseResult->Protocol == EProtocolFlavor::Unsync);
 		UNSYNC_ASSERT(ParseResult->StorageNamespace == "test.namespace");
 	}
@@ -304,8 +304,8 @@ TestParseRemote()
 		TResult<FRemoteDesc> ParseResult = FRemoteDesc::FromUrl("unsync://example.com:1234");
 		UNSYNC_ASSERT(ParseResult.IsOk());
 		UNSYNC_ASSERT(ParseResult->bTlsEnable == false);
-		UNSYNC_ASSERT(ParseResult->HostAddress == "example.com");
-		UNSYNC_ASSERT(ParseResult->HostPort == 1234);
+		UNSYNC_ASSERT(ParseResult->Host.Address == "example.com");
+		UNSYNC_ASSERT(ParseResult->Host.Port == 1234);
 		UNSYNC_ASSERT(ParseResult->Protocol == EProtocolFlavor::Unsync);
 	}
 
@@ -313,8 +313,8 @@ TestParseRemote()
 		TResult<FRemoteDesc> ParseResult = FRemoteDesc::FromUrl("unsync+tls://example.com:1234");
 		UNSYNC_ASSERT(ParseResult.IsOk());
 		UNSYNC_ASSERT(ParseResult->bTlsEnable == true);
-		UNSYNC_ASSERT(ParseResult->HostAddress == "example.com");
-		UNSYNC_ASSERT(ParseResult->HostPort == 1234);
+		UNSYNC_ASSERT(ParseResult->Host.Address == "example.com");
+		UNSYNC_ASSERT(ParseResult->Host.Port == 1234);
 		UNSYNC_ASSERT(ParseResult->Protocol == EProtocolFlavor::Unsync);
 	}
 
@@ -327,8 +327,8 @@ TestParseRemote()
 		TResult<FRemoteDesc> ParseResult = FRemoteDesc::FromUrl("unsync+tls://example.com:1234/request/path#namespace");
 		UNSYNC_ASSERT(ParseResult.IsOk());
 		UNSYNC_ASSERT(ParseResult->bTlsEnable == true);
-		UNSYNC_ASSERT(ParseResult->HostAddress == "example.com");
-		UNSYNC_ASSERT(ParseResult->HostPort == 1234);
+		UNSYNC_ASSERT(ParseResult->Host.Address == "example.com");
+		UNSYNC_ASSERT(ParseResult->Host.Port == 1234);
 		UNSYNC_ASSERT(ParseResult->Protocol == EProtocolFlavor::Unsync);
 		UNSYNC_ASSERT(ParseResult->StorageNamespace == "namespace");
 		UNSYNC_ASSERT(ParseResult->RequestPath == "request/path");
@@ -338,8 +338,8 @@ TestParseRemote()
 		TResult<FRemoteDesc> ParseResult = FRemoteDesc::FromUrl("http://example.com/request/path#namespace");
 		UNSYNC_ASSERT(ParseResult.IsOk());
 		UNSYNC_ASSERT(ParseResult->bTlsEnable == false);
-		UNSYNC_ASSERT(ParseResult->HostAddress == "example.com");
-		UNSYNC_ASSERT(ParseResult->HostPort == 80);
+		UNSYNC_ASSERT(ParseResult->Host.Address == "example.com");
+		UNSYNC_ASSERT(ParseResult->Host.Port == 80);
 		UNSYNC_ASSERT(ParseResult->Protocol == EProtocolFlavor::Jupiter);
 		UNSYNC_ASSERT(ParseResult->StorageNamespace == "namespace");
 		UNSYNC_ASSERT(ParseResult->RequestPath == "request/path");
diff --git a/Engine/Source/Programs/Unsync/Private/UnsyncRemote.h b/Engine/Source/Programs/Unsync/Private/UnsyncRemote.h
--- a/Engine/Source/Programs/Unsync/Private/UnsyncRemote.h
+++ b/Engine/Source/Programs/Unsync/Private/UnsyncRemote.h
@@ -6,6 +6,7 @@
 #include "UnsyncError.h"
 #include "UnsyncSocket.h"
 
+#include <optional>
 #include <memory>
 #include <string>
 
@@ -29,12 +30,19 @@ enum class ETransportProtocol {
 EProtocolFlavor ProtocolFlavorFromString(const char* Str);
 const char*		ToString(EProtocolFlavor Protocol);
 
+struct FHostAddressAndPort
+{
+	std::string Address;
+	uint16		Port = 0;
+
+	bool IsValid() const { return Port != 0 && !Address.empty(); }
+};
+
 struct FRemoteDesc
 {
 	EProtocolFlavor Protocol = EProtocolFlavor::Unknown;
 
-	std::string HostAddress;
-	uint16		HostPort = 0;
+	FHostAddressAndPort Host;
 
 	std::string RequestPath;
 	std::string StorageNamespace;
@@ -47,18 +55,19 @@ struct FRemoteDesc
 	std::string				 TlsSubjectOverride;			// Use host address if empty (default)
 	std::shared_ptr<FBuffer> TlsCacert;	 // Custom CA to use for server certificate validation (system root CA is used by default)
 
-	const std::string& GetTlsSubject() const { return TlsSubjectOverride.length() ? TlsSubjectOverride : HostAddress; }
+	const std::string& GetTlsSubject() const { return TlsSubjectOverride.length() ? TlsSubjectOverride : Host.Address; }
 
 	bool bAuthenticationRequired = false;
-	std::string LoginAddress; // Optional address of the server used for login requests. If empty, then HostAddress is used.
+	std::optional<FHostAddressAndPort> PrimaryHost;	 // Optional address of the server used for login requests and other queries. If
+													 // empty, then HostAddress is used.
 
-	const std::string& GetLoginAddress() const { return LoginAddress.length() ? LoginAddress : HostAddress; }
+	const FHostAddressAndPort& GetPrimaryHostAddress() const { return PrimaryHost ? *PrimaryHost : Host; }
 
 	uint32 RecvTimeoutSeconds = 0;
 
 	uint32 MaxConnections = 8;	// Limit on concurrent connections to this server
 
-	bool IsValid() const { return Protocol != EProtocolFlavor::Unknown && !HostAddress.empty() && HostPort != 0; }
+	bool IsValid() const { return Protocol != EProtocolFlavor::Unknown && Host.IsValid(); }
 
 	static TResult<FRemoteDesc> FromUrl(std::string_view Url);
 
