unsync - Fix metadata directory creation during pack process

diff --git a/Engine/Source/Programs/Unsync/Private/UnsyncCmdPack.cpp b/Engine/Source/Programs/Unsync/Private/UnsyncCmdPack.cpp
--- a/Engine/Source/Programs/Unsync/Private/UnsyncCmdPack.cpp
+++ b/Engine/Source/Programs/Unsync/Private/UnsyncCmdPack.cpp
@@ -344,6 +344,8 @@ CmdPack(const FCmdPackOptions& Options)
 	const FPath ManifestRoot = InputRoot / ".unsync";
 	const FPath StoreRoot	 = Options.StorePath.empty() ? ManifestRoot : Options.StorePath;
 	const FPath PackRoot	 = StoreRoot / "pack";
+	const FPath SnapshotRoot = StoreRoot / "snapshot";
+	const FPath TagRoot		 = StoreRoot / "tag";
 
 	UNSYNC_LOG(L"Generating package for directory '%ls'", InputRoot.wstring().c_str());
 	UNSYNC_LOG_INDENT;
@@ -405,6 +407,18 @@ CmdPack(const FCmdPackOptions& Options)
 			UNSYNC_ERROR(L"Failed to create pack output directory '%ls'", PackRoot.wstring().c_str());
 			return -1;
 		}
+
+		if (!EnsureDirectoryExists(SnapshotRoot))
+		{
+			UNSYNC_ERROR(L"Failed to create snapshot output directory '%ls'", SnapshotRoot.wstring().c_str());
+			return -1;
+		}
+
+		if (!EnsureDirectoryExists(TagRoot))
+		{
+			UNSYNC_ERROR(L"Failed to create tag output directory '%ls'", TagRoot.wstring().c_str());
+			return -1;
+		}
 	}
 
 	THashSet<FHash128> SeenBlockHashSet;  // TODO: use FPackDatabase, perhaps in some osrt of lightweight mode
@@ -658,7 +672,7 @@ CmdPack(const FCmdPackOptions& Options)
 		FHash128 ManifestBlocksBufferHash = HashBlake3Bytes<FHash128>(ManifestBlocksBuffer.Data, ManifestBlocksBuffer.Size);
 
 		std::string SnapshotId	 = HashToHexString(ManifestBlocksBufferHash);  // TODO: allow overriding this from command line
-		FPath		SnapshotPath = StoreRoot / "snapshot" / (SnapshotId + ".unsync_snapshot");
+		FPath		SnapshotPath = SnapshotRoot / (SnapshotId + ".unsync_snapshot");
 
 		UNSYNC_LOG(L"Writing snapshot: %hs", SnapshotId.c_str());
 
@@ -674,7 +688,7 @@ CmdPack(const FCmdPackOptions& Options)
 		if (!Options.SnapshotName.empty())
 		{
 			UNSYNC_LOG(L"Saving snapshot %hs as tag '%hs'", SnapshotId.c_str(), Options.SnapshotName.c_str());
-			FPath			TagPath = StoreRoot / "tag" / (Options.SnapshotName + ".unsync_tag");
+			FPath			TagPath = TagRoot / (Options.SnapshotName + ".unsync_tag");
 			std::error_code ErrorCode;
 			if (!FileCopyOverwrite(SnapshotPath, TagPath, ErrorCode))
 			{
diff --git a/Engine/Source/Programs/Unsync/Private/UnsyncCore.cpp b/Engine/Source/Programs/Unsync/Private/UnsyncCore.cpp
--- a/Engine/Source/Programs/Unsync/Private/UnsyncCore.cpp
+++ b/Engine/Source/Programs/Unsync/Private/UnsyncCore.cpp
@@ -27,7 +27,7 @@ UNSYNC_THIRD_PARTY_INCLUDES_START
 #include <md5-sse2.h>
 UNSYNC_THIRD_PARTY_INCLUDES_END
 
-#define UNSYNC_VERSION_STR "1.0.63"
+#define UNSYNC_VERSION_STR "1.0.64-dev1"
 
 namespace unsync {
 
