unsync - Add option to explicitly specify the auth token storage file for sync command

diff --git a/Engine/Source/Programs/Unsync/Private/UnsyncAuth.cpp b/Engine/Source/Programs/Unsync/Private/UnsyncAuth.cpp
--- a/Engine/Source/Programs/Unsync/Private/UnsyncAuth.cpp
+++ b/Engine/Source/Programs/Unsync/Private/UnsyncAuth.cpp
@@ -727,6 +727,11 @@ RefreshOrAcquireToken(const FAuthDesc& AuthDesc, const FOpenIdConfig& OpenIdConf
 TResult<FPath>
 GetTokenCachePath(const FAuthDesc& AuthDesc)
 {
+	if (!AuthDesc.TokenPath.empty())
+	{
+		return ResultOk(AuthDesc.TokenPath);
+	}
+
 	FPath UserHomePath = GetUserHomeDirectory();
 	if (UserHomePath.empty())
 	{
diff --git a/Engine/Source/Programs/Unsync/Private/UnsyncAuth.h b/Engine/Source/Programs/Unsync/Private/UnsyncAuth.h
--- a/Engine/Source/Programs/Unsync/Private/UnsyncAuth.h
+++ b/Engine/Source/Programs/Unsync/Private/UnsyncAuth.h
@@ -49,6 +49,8 @@ struct FAuthDesc
 	std::string Audience;
 	std::string Callback;	// OAuth flow redirection URL
 
+	FPath TokenPath; // Optional explicit storage location for the auth token
+
 	// Check that mandatory fields have values
 	bool IsValid() const
 	{
diff --git a/Engine/Source/Programs/Unsync/Private/UnsyncMain.cpp b/Engine/Source/Programs/Unsync/Private/UnsyncMain.cpp
--- a/Engine/Source/Programs/Unsync/Private/UnsyncMain.cpp
+++ b/Engine/Source/Programs/Unsync/Private/UnsyncMain.cpp
@@ -80,6 +80,7 @@ InnerMain(int Argc, char** Argv)
 	std::string				 P4HavePathUtf8;
 	std::string				 StorePathUtf8;
 	std::string				 SnapshotNameUtf8;
+	std::string				 AuthTokenPathUtf8;
 	bool					 bRunP4Have			 = false;
 	bool					 bForceOperation	 = false;
 	bool					 bAllowInsecureTls	 = false;
@@ -310,6 +311,7 @@ InnerMain(int Argc, char** Argv)
 	SubSync->add_flag("--no-space-validation", bNoSpaceValidation, "Skip checking available disk space before sync (DANGEROUS)");
 	SubSync->add_option("--scavenge", ScavengeRootUtf8, "Search for unsync manifests and reusable blocks in this directory (EXPERIMENTAL)");
 	SubSync->add_flag("--login", bShouldLogin, "Use user authentication when accessing unsync server");
+	SubSync->add_option("--token", AuthTokenPathUtf8, "Explicit path to the authentication token file to use");
 	SubSync->add_flag("--no-timeout", bNoSocketTimeout, "Disable the default 60 second timeout on network socket operations");
 
 	CLI::Option* BackgroundMemoryBudgetOption = SubSync->add_option("--background-task-memory",
@@ -867,6 +869,8 @@ InnerMain(int Argc, char** Argv)
 			{
 				AuthDesc = AuthDescResult.GetData();
 
+				AuthDesc.TokenPath = NormalizeFilenameUtf8(AuthTokenPathUtf8);
+
 				// Note: since tokens can expire during a long operation,
 				// we can only save the auth descriptor and re-authenticate later if necessary
 				TResult<FAuthToken> AuthTokenResult = Authenticate(AuthDesc, 5 * 60);
@@ -906,6 +910,10 @@ InnerMain(int Argc, char** Argv)
 				UNSYNC_VERBOSE2(L"Detected memory: %llu GB", InstalledMemoryGB);
 				BackgroundTaskMemoryBudgetGB = std::max<uint32>(2, InstalledMemoryGB / 4);
 			}
+			else
+			{
+				UNSYNC_VERBOSE2(L"Could not detect system memory size");
+			}
 
 			UNSYNC_VERBOSE2(L"Using automatic background task memory budget: %llu GB", BackgroundTaskMemoryBudgetGB);
 		}
