unsync - Always save downloaded manifest file, even if the download source is not a filesystem

diff --git a/Engine/Source/Programs/Unsync/Private/UnsyncCore.cpp b/Engine/Source/Programs/Unsync/Private/UnsyncCore.cpp
--- a/Engine/Source/Programs/Unsync/Private/UnsyncCore.cpp
+++ b/Engine/Source/Programs/Unsync/Private/UnsyncCore.cpp
@@ -27,7 +27,7 @@ UNSYNC_THIRD_PARTY_INCLUDES_START
 #include <md5-sse2.h>
 UNSYNC_THIRD_PARTY_INCLUDES_END
 
-#define UNSYNC_VERSION_STR "1.0.64-dev1"
+#define UNSYNC_VERSION_STR "1.0.64"
 
 namespace unsync {
 
@@ -3004,7 +3004,7 @@ SyncDirectory(const FSyncDirectoryOptions& SyncOptions)
 
 	// Save the source directory manifest on success.
 	// It can be used to speed up the diffing process during next sync.
-	if (bFileSystemSource && bSyncSucceeded && !GDryRun)
+	if (bSyncSucceeded && !GDryRun)
 	{
 		bool bSaveOk = SaveDirectoryManifest(SourceDirectoryManifest, TargetManifestPath);
 
