unsync - Fix manifest version field initialization

diff --git a/Engine/Source/Programs/Unsync/Private/UnsyncManifest.cpp b/Engine/Source/Programs/Unsync/Private/UnsyncManifest.cpp
--- a/Engine/Source/Programs/Unsync/Private/UnsyncManifest.cpp
+++ b/Engine/Source/Programs/Unsync/Private/UnsyncManifest.cpp
@@ -283,6 +283,7 @@ CreateDirectoryManifest(const FPath& Root, const FComputeBlocksParams& Params)
 	FDirectoryManifest Result;
 
 	Result.Algorithm = Params.Algorithm;
+	Result.Version	 = FDirectoryManifest::VERSION;
 
 	FTimePoint TimeBegin = TimePointNow();
 
