unsync - Automatically enable TLS when server port is standard https/443

diff --git a/Engine/Source/Programs/Unsync/Private/UnsyncCmdQuery.cpp b/Engine/Source/Programs/Unsync/Private/UnsyncCmdQuery.cpp
--- a/Engine/Source/Programs/Unsync/Private/UnsyncCmdQuery.cpp
+++ b/Engine/Source/Programs/Unsync/Private/UnsyncCmdQuery.cpp
@@ -76,6 +76,10 @@ RunQueryMirrors(const FRemoteDesc& RemoteDesc)
 			{
 				Info.Name = Field.second.string_value();
 			}
+			else if (Field.first == "description")
+			{
+				Info.Description = Field.second.string_value();
+			}
 			else if (Field.first == "address")
 			{
 				Info.Address = Field.second.string_value();
@@ -126,13 +130,16 @@ CmdQueryMirrors(const FCmdQueryOptions& Options)
 	{
 		const FMirrorInfo& Mirror = Mirrors[I];
 
+		int32 PingMs = (Mirror.Ping == 0) ? 0 : std::max(1, int32(Mirror.Ping * 1000.0));
+
 		LogPrintf(ELogLevel::MachineReadable,
-				  L"  {\"address\":\"%hs\", \"port\":%d, \"ok\":%hs, \"ping\":%d, \"name\":\"%hs\"}%hs\n",
+				  L"  {\"address\":\"%hs\", \"port\":%d, \"ok\":%hs, \"ping\":%d, \"name\":\"%hs\", \"description\":\"%hs\"}%hs\n",
 				  StringEscape(Mirror.Address).c_str(),
 				  Mirror.Port,
 				  Mirror.Ping > 0 ? "true" : "false",
-				  int32(Mirror.Ping * 1000.0),
+				  PingMs,
 				  StringEscape(Mirror.Name).c_str(),
+				  StringEscape(Mirror.Description).c_str(),
 				  I + 1 == Mirrors.size() ? "" : ",");
 	}
 
diff --git a/Engine/Source/Programs/Unsync/Private/UnsyncCmdQuery.h b/Engine/Source/Programs/Unsync/Private/UnsyncCmdQuery.h
--- a/Engine/Source/Programs/Unsync/Private/UnsyncCmdQuery.h
+++ b/Engine/Source/Programs/Unsync/Private/UnsyncCmdQuery.h
@@ -21,6 +21,7 @@ struct FMirrorInfo
 {
 	std::string Name;
 	std::string Address;
+	std::string Description;
 	uint16		Port = UNSYNC_DEFAULT_PORT;
 	double		Ping = 0;
 };
diff --git a/Engine/Source/Programs/Unsync/Private/UnsyncCore.cpp b/Engine/Source/Programs/Unsync/Private/UnsyncCore.cpp
--- a/Engine/Source/Programs/Unsync/Private/UnsyncCore.cpp
+++ b/Engine/Source/Programs/Unsync/Private/UnsyncCore.cpp
@@ -27,7 +27,7 @@ UNSYNC_THIRD_PARTY_INCLUDES_START
 #include <md5-sse2.h>
 UNSYNC_THIRD_PARTY_INCLUDES_END
 
-#define UNSYNC_VERSION_STR "1.0.64"
+#define UNSYNC_VERSION_STR "1.0.65"
 
 namespace unsync {
 
diff --git a/Engine/Source/Programs/Unsync/Private/UnsyncProxy.cpp b/Engine/Source/Programs/Unsync/Private/UnsyncProxy.cpp
--- a/Engine/Source/Programs/Unsync/Private/UnsyncProxy.cpp
+++ b/Engine/Source/Programs/Unsync/Private/UnsyncProxy.cpp
@@ -968,10 +968,12 @@ FProxyPool::FProxyPool(const FRemoteDesc& InRemoteDesc, const FAuthDesc* InAuthD
 		else
 		{
 			const ProxyQuery::FHelloResponse& Data = Response.GetData();
-			UNSYNC_VERBOSE(L"Connection established. Server name: %hs, version: %hs, git: %hs.",
+			UNSYNC_VERBOSE(L"Connection established. Server name: %hs, version: %hs, git: %hs, tls: %hs.",
 						   Data.Name.empty() ? "unknown" : Data.Name.c_str(),
 						   Data.VersionNumber.empty() ? "unknown" : Data.VersionNumber.c_str(),
-						   Data.VersionGit.empty() ? "unknown" : Data.VersionGit.c_str());
+						   Data.VersionGit.empty() ? "unknown" : Data.VersionGit.c_str(),
+						   RemoteDesc.bTlsEnable ? "yes" : "no"
+				);
 
 			Features  = Data.Features;
 			SessionId = Data.SessionId;
diff --git a/Engine/Source/Programs/Unsync/Private/UnsyncRemote.cpp b/Engine/Source/Programs/Unsync/Private/UnsyncRemote.cpp
--- a/Engine/Source/Programs/Unsync/Private/UnsyncRemote.cpp
+++ b/Engine/Source/Programs/Unsync/Private/UnsyncRemote.cpp
@@ -112,6 +112,8 @@ FRemoteDesc::FromUrl(std::string_view Url)
 		}
 	}
 
+	const size_t NamespacePos = HostAddress.find_last_of('#');
+
 	switch (Transport)
 	{
 		default:
@@ -119,11 +121,17 @@ FRemoteDesc::FromUrl(std::string_view Url)
 			Result.Protocol = EProtocolFlavor::Unsync;
 			break;
 		case ETransportProtocol::Http:
-			Result.Protocol = EProtocolFlavor::Jupiter;	 // TODO: UNSYNC-over-HTTP could also be supported
+			if (NamespacePos != std::string::npos || Scheme.starts_with("jupiter"))
+			{
+				Result.Protocol = EProtocolFlavor::Jupiter;
+			}
+			else
+			{
+				Result.Protocol = EProtocolFlavor::Unsync;
+			}
 			break;
 	}
 
-	size_t NamespacePos = HostAddress.find_last_of('#');
 	if (NamespacePos != std::string::npos)
 	{
 		Result.StorageNamespace = HostAddress.substr(NamespacePos + 1);
@@ -131,9 +139,9 @@ FRemoteDesc::FromUrl(std::string_view Url)
 	}
 
 	uint16 HostPort = 0;
-	size_t PortPos	= HostAddress.find_first_of(':');
+	const size_t PortPos	= HostAddress.find_first_of(':');
 
-	size_t RequestPos = HostAddress.find_first_of('/');
+	const size_t RequestPos = HostAddress.find_first_of('/');
 	if (RequestPos != std::string::npos)
 	{
 		Result.RequestPath = HostAddress.substr(RequestPos + 1);
@@ -163,6 +171,11 @@ FRemoteDesc::FromUrl(std::string_view Url)
 		HostAddress = HostAddress.substr(0, PortPos);
 	}
 
+	if (HostPort == 443)
+	{
+		bUseTls = true;
+	}
+
 	if (HostPort == 0)
 	{
 		return AppError("Invalid host port");  // TODO: extract the port substring
@@ -217,7 +230,7 @@ TestParseRemote()
 	}
 
 	{
-		TResult<FRemoteDesc> ParseResult = FRemoteDesc::FromUrl("http://example.com");
+		TResult<FRemoteDesc> ParseResult = FRemoteDesc::FromUrl("http://example.com#foo");
 		UNSYNC_ASSERT(ParseResult.IsOk());
 		UNSYNC_ASSERT(ParseResult->bTlsEnable == false);
 		UNSYNC_ASSERT(ParseResult->Host.Address == "example.com");
@@ -226,7 +239,7 @@ TestParseRemote()
 	}
 
 	{
-		TResult<FRemoteDesc> ParseResult = FRemoteDesc::FromUrl("https://example.com");
+		TResult<FRemoteDesc> ParseResult = FRemoteDesc::FromUrl("https://example.com#foo");
 		UNSYNC_ASSERT(ParseResult.IsOk());
 		UNSYNC_ASSERT(ParseResult->bTlsEnable == true);
 		UNSYNC_ASSERT(ParseResult->Host.Address == "example.com");
@@ -235,7 +248,7 @@ TestParseRemote()
 	}
 
 	{
-		TResult<FRemoteDesc> ParseResult = FRemoteDesc::FromUrl("http://example.com:1234");
+		TResult<FRemoteDesc> ParseResult = FRemoteDesc::FromUrl("http://example.com:1234#foo");
 		UNSYNC_ASSERT(ParseResult.IsOk());
 		UNSYNC_ASSERT(ParseResult->bTlsEnable == false);
 		UNSYNC_ASSERT(ParseResult->Host.Address == "example.com");
@@ -244,7 +257,7 @@ TestParseRemote()
 	}
 
 	{
-		TResult<FRemoteDesc> ParseResult = FRemoteDesc::FromUrl("https://example.com:1234");
+		TResult<FRemoteDesc> ParseResult = FRemoteDesc::FromUrl("https://example.com:1234#foo");
 		UNSYNC_ASSERT(ParseResult.IsOk());
 		UNSYNC_ASSERT(ParseResult->bTlsEnable == true);
 		UNSYNC_ASSERT(ParseResult->Host.Address == "example.com");
@@ -344,6 +357,15 @@ TestParseRemote()
 		UNSYNC_ASSERT(ParseResult->StorageNamespace == "namespace");
 		UNSYNC_ASSERT(ParseResult->RequestPath == "request/path");
 	}
+
+	{
+		TResult<FRemoteDesc> ParseResult = FRemoteDesc::FromUrl("example.com:443");
+		UNSYNC_ASSERT(ParseResult.IsOk());
+		UNSYNC_ASSERT(ParseResult->bTlsEnable == true);
+		UNSYNC_ASSERT(ParseResult->Host.Address == "example.com");
+		UNSYNC_ASSERT(ParseResult->Host.Port == 443);
+		UNSYNC_ASSERT(ParseResult->Protocol == EProtocolFlavor::Unsync);
+	}
 }
 
 }  // namespace unsync
