unsync - Attempt to establish TLS connections by default and allow fallback unless explicitly connecting to HTTPS port or using authentication

diff --git a/Engine/Source/Programs/Unsync/Private/UnsyncCmdPush.cpp b/Engine/Source/Programs/Unsync/Private/UnsyncCmdPush.cpp
--- a/Engine/Source/Programs/Unsync/Private/UnsyncCmdPush.cpp
+++ b/Engine/Source/Programs/Unsync/Private/UnsyncCmdPush.cpp
@@ -73,8 +73,6 @@ CmdPush(const FCmdPushOptions& Options)
 			}
 		}
 
-		const bool bUseTls = Options.Remote.bTlsEnable;
-
 		FTlsClientSettings TlsSettings = Options.Remote.GetTlsClientSettings();
 
 		if (Options.Remote.StorageNamespace.empty())
@@ -89,7 +87,7 @@ CmdPush(const FCmdPushOptions& Options)
 		bool		 bPushComplete = false;
 		for (uint32 AttemptIndex = 0; AttemptIndex < MaxAttempts; ++AttemptIndex)
 		{
-			TResult<uint64> PushResult = JupiterPush(Manifest, Options.Remote, bUseTls ? &TlsSettings : nullptr);
+			TResult<uint64> PushResult = JupiterPush(Manifest, Options.Remote, TlsSettings);
 
 			if (const uint64* PushedBlocks = PushResult.TryData())
 			{
diff --git a/Engine/Source/Programs/Unsync/Private/UnsyncCmdQuery.cpp b/Engine/Source/Programs/Unsync/Private/UnsyncCmdQuery.cpp
--- a/Engine/Source/Programs/Unsync/Private/UnsyncCmdQuery.cpp
+++ b/Engine/Source/Programs/Unsync/Private/UnsyncCmdQuery.cpp
@@ -25,11 +25,27 @@ using FMirrorInfoResult = TResult<std::vector<FMirrorInfo>>;
 static double
 RunHttpPing(std::string_view Address, uint16 Port)
 {
-	FTimePoint		TimeBegin = TimePointNow();
-	FHttpConnection Connection(Address, Port);
+	
+	FTlsClientSettings TlsSettings;
+	TlsSettings.Subject = Address;
+
+	ETlsRequirement TlsRequirement = ETlsRequirement::None;
+
+	if (Port == 443)
+	{
+		TlsRequirement = ETlsRequirement::Preferred;
+	}
+
+	FHttpConnection Connection(Address, Port, TlsRequirement, TlsSettings);
+
 	FHttpRequest	Request;
 	Request.Url				   = "/api/v1/ping";
 	Request.Method			   = EHttpMethod::GET;
+
+	// Don't time the connection handshake, only query time itself
+	Connection.Open();
+
+	FTimePoint	  TimeBegin	   = TimePointNow();
 	FHttpResponse PingResponse = HttpRequest(Connection, Request);
 	FTimePoint	  TimeEnd	   = TimePointNow();
 
@@ -215,7 +231,7 @@ CmdQuerySearch(const FCmdQueryOptions& Options)
 	auto CreateConnection = [Remote = Options.Remote]
 	{
 		FTlsClientSettings TlsSettings = Remote.GetTlsClientSettings();
-		return new FHttpConnection(Remote.Host.Address, Remote.Host.Port, &TlsSettings);
+		return new FHttpConnection(Remote.Host.Address, Remote.Host.Port, Remote.TlsRequirement, TlsSettings);
 	};
 
 	TObjectPool<FHttpConnection> ConnectionPool(CreateConnection);
diff --git a/Engine/Source/Programs/Unsync/Private/UnsyncCore.cpp b/Engine/Source/Programs/Unsync/Private/UnsyncCore.cpp
--- a/Engine/Source/Programs/Unsync/Private/UnsyncCore.cpp
+++ b/Engine/Source/Programs/Unsync/Private/UnsyncCore.cpp
@@ -27,7 +27,7 @@ UNSYNC_THIRD_PARTY_INCLUDES_START
 #include <md5-sse2.h>
 UNSYNC_THIRD_PARTY_INCLUDES_END
 
-#define UNSYNC_VERSION_STR "1.0.65"
+#define UNSYNC_VERSION_STR "1.0.66"
 
 namespace unsync {
 
diff --git a/Engine/Source/Programs/Unsync/Private/UnsyncHttp.cpp b/Engine/Source/Programs/Unsync/Private/UnsyncHttp.cpp
--- a/Engine/Source/Programs/Unsync/Private/UnsyncHttp.cpp
+++ b/Engine/Source/Programs/Unsync/Private/UnsyncHttp.cpp
@@ -259,7 +259,7 @@ HttpRequest(const FRemoteDesc& RemoteDesc,
 			std::string_view   BearerToken)
 {
 	FTlsClientSettings TlsSettings = RemoteDesc.GetTlsClientSettings();
-	FHttpConnection	   Connection(RemoteDesc.Host.Address, RemoteDesc.Host.Port, RemoteDesc.bTlsEnable ? &TlsSettings : nullptr);
+	FHttpConnection	   Connection(RemoteDesc.Host.Address, RemoteDesc.Host.Port, RemoteDesc.TlsRequirement, TlsSettings);
 
 	FHttpRequest Request;
 
@@ -479,42 +479,44 @@ HttpRequestEnd(FHttpConnection& Connection)
 
 	Connection.NumActiveRequests -= 1;
 
+	Result.bConnectionEncrypted = Connection.IsEncrypted();
+
 	// TODO: report errors
 
 	return Result;
 }
 
-FHttpConnection::FHttpConnection(const std::string_view InHostAddress, uint16 InPort, const FTlsClientSettings* InTlsSettings)
+FHttpConnection::FHttpConnection(const std::string_view	   InHostAddress,
+								 uint16					   InPort,
+								 ETlsRequirement		   InTlsRequirement,
+								 const FTlsClientSettings& InTlsSettings)
 : HostAddress(InHostAddress)
 , HostPort(InPort)
-, bUseTls(InTlsSettings != nullptr)
+, TlsRequirement(InTlsRequirement)
 {
-	if (InTlsSettings)
+	if (InTlsSettings.Subject.empty())
 	{
-		if (InTlsSettings->Subject.empty())
-		{
-			TlsSubject = std::string(InHostAddress);
-		}
-		else
-		{
-			TlsSubject = std::string(InTlsSettings->Subject);
-		}
+		TlsSubject = std::string(InHostAddress);
+	}
+	else
+	{
+		TlsSubject = std::string(InTlsSettings.Subject);
+	}
 
-		bTlsVerifyCertificate = InTlsSettings->bVerifyCertificate;
-		bTlsVerifySubject	  = InTlsSettings->bVerifySubject;
-		if (InTlsSettings->CACert.Data)
-		{
-			TlsCacert = std::make_shared<FBuffer>();
-			TlsCacert->Append(InTlsSettings->CACert.Data, InTlsSettings->CACert.Size);
-		}
+	bTlsVerifyCertificate = InTlsSettings.bVerifyCertificate;
+	bTlsVerifySubject	  = InTlsSettings.bVerifySubject;
+	if (InTlsSettings.CACert.Data)
+	{
+		TlsCacert = std::make_shared<FBuffer>();
+		TlsCacert->Append(InTlsSettings.CACert.Data, InTlsSettings.CACert.Size);
 	}
 }
 
 FHttpConnection::FHttpConnection(const FHttpConnection& Other)
 : HostAddress(Other.HostAddress)
 , HostPort(Other.HostPort)
-, bUseTls(Other.bUseTls)
 , bKeepAlive(Other.bKeepAlive)
+, TlsRequirement(Other.TlsRequirement)
 , bTlsVerifySubject(Other.bTlsVerifySubject)
 , TlsSubject(Other.TlsSubject)
 , bTlsVerifyCertificate(Other.bTlsVerifyCertificate)
@@ -525,7 +527,7 @@ FHttpConnection::FHttpConnection(const FHttpConnection& Other)
 FHttpConnection
 FHttpConnection::CreateDefaultHttp(const std::string_view InHostAddress, uint16 Port)
 {
-	return FHttpConnection(InHostAddress, Port, nullptr);
+	return FHttpConnection(InHostAddress, Port);
 }
 
 FHttpConnection
@@ -533,14 +535,14 @@ FHttpConnection::CreateDefaultHttps(const std::string_view InHostAddress, uint16
 {
 	FTlsClientSettings TlsSettings;
 	TlsSettings.Subject = InHostAddress.data();
-	return FHttpConnection(InHostAddress, Port, &TlsSettings);
+	return FHttpConnection(InHostAddress, Port, ETlsRequirement::Required, TlsSettings);
 }
 
 FHttpConnection
 FHttpConnection::CreateDefaultHttps(const FRemoteDesc& RemoteDesc)
 {
 	FTlsClientSettings TlsSettings = RemoteDesc.GetTlsClientSettings();
-	return FHttpConnection(RemoteDesc.Host.Address, RemoteDesc.Host.Port, &TlsSettings);
+	return FHttpConnection(RemoteDesc.Host.Address, RemoteDesc.Host.Port, ETlsRequirement::Required, TlsSettings);
 }
 
 bool
@@ -554,15 +556,10 @@ FHttpConnection::Open()
 		}
 	}
 
-	FSocketHandle RawSocketHandle = SocketConnectTcp(HostAddress.c_str(), HostPort);
-
-	if (RawSocketHandle == InvalidSocketHandle)
+	if (TlsRequirement != ETlsRequirement::None)
 	{
-		return false;
-	}
+		FSocketHandle RawSocketHandle = SocketConnectTcp(HostAddress.c_str(), HostPort);
 
-	if (bUseTls)
-	{
 		FTlsClientSettings ClientSettings;
 		ClientSettings.bVerifyCertificate = bTlsVerifyCertificate;
 		ClientSettings.bVerifySubject	  = bTlsVerifySubject;
@@ -574,6 +571,7 @@ FHttpConnection::Open()
 		}
 
 		FSocketTls* TlsSocket = new FSocketTls(RawSocketHandle, ClientSettings);
+
 		if (TlsSocket->IsTlsValid())
 		{
 			Socket = std::unique_ptr<FSocketTls>(TlsSocket);
@@ -581,10 +579,14 @@ FHttpConnection::Open()
 		else
 		{
 			delete TlsSocket;
+			TlsSocket = nullptr;
 		}
 	}
-	else
+
+	if (!Socket && TlsRequirement != ETlsRequirement::Required)
 	{
+		FSocketHandle RawSocketHandle = SocketConnectTcp(HostAddress.c_str(), HostPort);
+
 		Socket = std::unique_ptr<FSocketRaw>(new FSocketRaw(RawSocketHandle));
 	}
 
diff --git a/Engine/Source/Programs/Unsync/Private/UnsyncHttp.h b/Engine/Source/Programs/Unsync/Private/UnsyncHttp.h
--- a/Engine/Source/Programs/Unsync/Private/UnsyncHttp.h
+++ b/Engine/Source/Programs/Unsync/Private/UnsyncHttp.h
@@ -55,6 +55,8 @@ struct FHttpResponse
 
 	EHttpContentType ContentType = EHttpContentType::Unknown;
 
+	bool bConnectionEncrypted = false;
+
 	bool Success() const { return Code >= 200 && Code < 300; }
 
 	std::string_view AsStringView() const { return std::string_view((const char*)Buffer.Data(), Buffer.Size()); }
@@ -62,7 +64,8 @@ struct FHttpResponse
 
 struct FHttpConnection
 {
-	FHttpConnection(const std::string_view InHostAddress, uint16 InPort, const FTlsClientSettings* TlsSettings = nullptr);
+	FHttpConnection(const std::string_view InHostAddress, uint16 InPort) : FHttpConnection(InHostAddress, InPort, ETlsRequirement::None, FTlsClientSettings{}) {}
+	FHttpConnection(const std::string_view InHostAddress, uint16 InPort, ETlsRequirement InTlsRequirement, const FTlsClientSettings& InTlsSettings);
 	FHttpConnection(const FHttpConnection& Other);
 
 	[[nodiscard]] static FHttpConnection CreateDefaultHttp(const std::string_view InHostAddress, uint16 Port = 80);
@@ -72,12 +75,12 @@ struct FHttpConnection
 	bool Open();
 	void Close();
 
-	const std::string HostAddress;		 // NOLINT
-	const uint16	  HostPort = 80;	 // NOLINT
-	const bool		  bUseTls  = false;	 // NOLINT
+	std::string		HostAddress;
+	uint16			HostPort	   = 80;
 
 	bool bKeepAlive = true;
 
+	ETlsRequirement			 TlsRequirement	   = ETlsRequirement::None;
 	bool					 bTlsVerifySubject = true;
 	std::string				 TlsSubject;
 	bool					 bTlsVerifyCertificate = true;
@@ -98,6 +101,8 @@ struct FHttpConnection
 		return *Socket;
 	}
 
+	bool IsEncrypted() const { return Socket && Socket->IsEncrypted(); }
+
 protected:
 	std::unique_ptr<FSocketBase> Socket;
 };
diff --git a/Engine/Source/Programs/Unsync/Private/UnsyncJupiter.cpp b/Engine/Source/Programs/Unsync/Private/UnsyncJupiter.cpp
--- a/Engine/Source/Programs/Unsync/Private/UnsyncJupiter.cpp
+++ b/Engine/Source/Programs/Unsync/Private/UnsyncJupiter.cpp
@@ -464,10 +464,10 @@ JupiterPutManifest(const FDirectoryManifest&	 Manifest,
 }
 
 TResult<uint64>
-JupiterPush(const FDirectoryManifest& Manifest, const FRemoteDesc& RemoteDesc, FTlsClientSettings* TlsSettings)
+JupiterPush(const FDirectoryManifest& Manifest, const FRemoteDesc& RemoteDesc, const FTlsClientSettings& TlsSettings)
 {
 	auto CreateConnection = [RemoteDesc, TlsSettings] {
-		return new FHttpConnection(RemoteDesc.Host.Address, RemoteDesc.Host.Port, TlsSettings);
+		return new FHttpConnection(RemoteDesc.Host.Address, RemoteDesc.Host.Port, RemoteDesc.TlsRequirement, TlsSettings);
 	};
 
 	FSemaphore					 ChunkUploadSemaphore(8);  // up to 8 concurrent connections
@@ -697,10 +697,9 @@ JupiterGetRawBlob(FHttpConnection& Connection, std::string_view JupiterNamespace
 
 FJupiterProtocolImpl::FJupiterProtocolImpl(const FRemoteDesc&		 InSettings,
 										   const FBlockRequestMap*	 InRequestMap,
-										   const FTlsClientSettings* TlsSettings,
 										   std::string_view			 InHttpHeaders)
 : FRemoteProtocolBase(InSettings, InRequestMap)
-, Connection(InSettings.Host.Address, InSettings.Host.Port, InSettings.bTlsEnable ? TlsSettings : nullptr)
+, Connection(InSettings.Host.Address, InSettings.Host.Port, InSettings.TlsRequirement, InSettings.GetTlsClientSettings())
 , HttpHeaders(InHttpHeaders)
 , RemoteDesc(InSettings)
 {
diff --git a/Engine/Source/Programs/Unsync/Private/UnsyncJupiter.h b/Engine/Source/Programs/Unsync/Private/UnsyncJupiter.h
--- a/Engine/Source/Programs/Unsync/Private/UnsyncJupiter.h
+++ b/Engine/Source/Programs/Unsync/Private/UnsyncJupiter.h
@@ -18,7 +18,7 @@ struct FDirectoryManifest;
 struct FTlsClientSettings;
 
 // Returns number of blocks pushed
-TResult<uint64> JupiterPush(const FDirectoryManifest& Manifest, const FRemoteDesc& RemoteDesc, FTlsClientSettings* TlsSettings = nullptr);
+TResult<uint64> JupiterPush(const FDirectoryManifest& Manifest, const FRemoteDesc& RemoteDesc, const FTlsClientSettings& TlsSettings);
 
 TResult<> JupiterPutRawBlob(FHttpConnection&	   Connection,
 							const std::string_view BaseUrl,
@@ -37,7 +37,6 @@ struct FJupiterProtocolImpl : FRemoteProtocolBase
 {
 	FJupiterProtocolImpl(const FRemoteDesc&		   InSettings,
 						 const FBlockRequestMap*   InRequestMap,
-						 const FTlsClientSettings* TlsSettings,
 						 std::string_view		   HttpHeaders);
 	virtual bool			 IsValid() const override;
 	virtual TResult<FBuffer> DownloadManifest(std::string_view ManifestName) override;
diff --git a/Engine/Source/Programs/Unsync/Private/UnsyncMain.cpp b/Engine/Source/Programs/Unsync/Private/UnsyncMain.cpp
--- a/Engine/Source/Programs/Unsync/Private/UnsyncMain.cpp
+++ b/Engine/Source/Programs/Unsync/Private/UnsyncMain.cpp
@@ -84,7 +84,7 @@ InnerMain(int Argc, char** Argv)
 	bool					 bRunP4Have			 = false;
 	bool					 bForceOperation	 = false;
 	bool					 bAllowInsecureTls	 = false;
-	bool					 bUseTls			 = false;
+	bool					 bRequireTls		 = false;
 	bool					 bUseDebugMode		 = false;
 	bool					 bIncrementalMode	 = false;
 	bool					 bNoOutputValidation = false;
@@ -118,9 +118,9 @@ InnerMain(int Argc, char** Argv)
 
 	const std::string HiddenGroupId;  // CLI11 uses an empty string group name to mark arguments that should be hidden
 
-	auto AddTlsOptions = [&CacertFilenameUtf8, &bUseTls, &bAllowInsecureTls](CLI::App* App) {
+	auto AddTlsOptions = [&CacertFilenameUtf8, &bRequireTls, &bAllowInsecureTls](CLI::App* App) {
 		App->add_option("--cacert", CacertFilenameUtf8, "Certificate authority file to use for TLS validation (.pem)");
-		App->add_flag("--tls", bUseTls, "Use TLS when connecting to remote server");
+		App->add_flag("--tls", bRequireTls, "Force TLS when connecting to remote server");
 		App->add_flag("--insecure", bAllowInsecureTls, "Skip remote server TLS certificate validation");
 	};
 
@@ -719,9 +719,9 @@ InnerMain(int Argc, char** Argv)
 		}
 	}
 
-	if (bUseTls)
+	if (bRequireTls)
 	{
-		RemoteDesc.bTlsEnable = true;
+		RemoteDesc.TlsRequirement = ETlsRequirement::Required;
 	}
 
 	if (bAllowInsecureTls)
@@ -787,6 +787,15 @@ InnerMain(int Argc, char** Argv)
 
 			RemoteDesc.Host.Address = Mirror->Address;
 			RemoteDesc.Host.Port	= Mirror->Port;
+			
+			if (Mirror->Port == 443)
+			{
+				RemoteDesc.TlsRequirement = ETlsRequirement::Required;
+			}
+			else if (RemoteDesc.TlsRequirement < ETlsRequirement::Required)
+			{
+				RemoteDesc.TlsRequirement = ETlsRequirement::Preferred;
+			}
 		}
 		else
 		{
@@ -883,7 +892,7 @@ InnerMain(int Argc, char** Argv)
 				}
 
 				// Authentication requires encrypted connection
-				RemoteDesc.bTlsEnable			   = true;
+				RemoteDesc.TlsRequirement		   = ETlsRequirement::Required;
 				RemoteDesc.bAuthenticationRequired = true;
 
 				UNSYNC_LOG(L"Authentication enabled")
diff --git a/Engine/Source/Programs/Unsync/Private/UnsyncProxy.cpp b/Engine/Source/Programs/Unsync/Private/UnsyncProxy.cpp
--- a/Engine/Source/Programs/Unsync/Private/UnsyncProxy.cpp
+++ b/Engine/Source/Programs/Unsync/Private/UnsyncProxy.cpp
@@ -19,8 +19,7 @@ struct FUnsyncProtocolImpl : FRemoteProtocolBase
 	FUnsyncProtocolImpl(const FRemoteDesc&			   InRemoteDesc,
 						const FRemoteProtocolFeatures& InFeatures,
 						const FAuthDesc*			   InAuthDesc,
-						const FBlockRequestMap*		   InRequestMap,
-						const FTlsClientSettings*	   TlsSettings);
+						const FBlockRequestMap*		   InRequestMap);
 	virtual ~FUnsyncProtocolImpl() override;
 	virtual bool			 IsValid() const override;
 	virtual FDownloadResult	 Download(const TArrayView<FNeedBlock> NeedBlocks, const FBlockDownloadCallback& CompletionCallback) override;
@@ -45,16 +44,14 @@ FProxy::FProxy(const FRemoteDesc& RemoteDesc, const FRemoteProtocolFeatures& InF
 {
 	UNSYNC_ASSERT(InRequestMap);
 
-	FTlsClientSettings TlsSettings = RemoteDesc.GetTlsClientSettings();
-
 	if (RemoteDesc.Protocol == EProtocolFlavor::Jupiter)
 	{
-		auto Inner	 = new FJupiterProtocolImpl(RemoteDesc, InRequestMap, &TlsSettings, RemoteDesc.HttpHeaders);
+		auto Inner	 = new FJupiterProtocolImpl(RemoteDesc, InRequestMap, RemoteDesc.HttpHeaders);
 		ProtocolImpl = std::unique_ptr<FRemoteProtocolBase>(Inner);
 	}
 	else if (RemoteDesc.Protocol == EProtocolFlavor::Unsync)
 	{
-		auto* Inner	 = new FUnsyncProtocolImpl(RemoteDesc, InFeatures, InAuthDesc, InRequestMap, &TlsSettings);
+		auto* Inner	 = new FUnsyncProtocolImpl(RemoteDesc, InFeatures, InAuthDesc, InRequestMap);
 		ProtocolImpl = std::unique_ptr<FRemoteProtocolBase>(Inner);
 	}
 	else
@@ -66,19 +63,19 @@ FProxy::FProxy(const FRemoteDesc& RemoteDesc, const FRemoteProtocolFeatures& InF
 FUnsyncProtocolImpl::FUnsyncProtocolImpl(const FRemoteDesc&				RemoteDesc,
 										 const FRemoteProtocolFeatures& InFeatures,
 										 const FAuthDesc*				InAuthDesc,
-										 const FBlockRequestMap*		InRequestMap,
-										 const FTlsClientSettings*		TlsSettings)
+										 const FBlockRequestMap*		InRequestMap)
 : FRemoteProtocolBase(RemoteDesc, InRequestMap)
 , Features(InFeatures)
 {
-	if (RemoteDesc.bTlsEnable && TlsSettings)
+	if (RemoteDesc.TlsRequirement != ETlsRequirement::None)
 	{
+		FTlsClientSettings TlsSettings = RemoteDesc.GetTlsClientSettings();
 		FSocketHandle RawSocketHandle = SocketConnectTcp(RemoteDesc.Host.Address.c_str(), RemoteDesc.Host.Port);
 		SocketSetRecvTimeout(RawSocketHandle, RemoteDesc.RecvTimeoutSeconds);
 
 		if (RawSocketHandle)
 		{
-			FSocketTls* TlsSocket = new FSocketTls(RawSocketHandle, *TlsSettings);
+			FSocketTls* TlsSocket = new FSocketTls(RawSocketHandle, TlsSettings);
 			if (TlsSocket->IsTlsValid())
 			{
 				SocketHandle = std::unique_ptr<FSocketTls>(TlsSocket);
@@ -90,7 +87,7 @@ FUnsyncProtocolImpl::FUnsyncProtocolImpl(const FRemoteDesc&				RemoteDesc,
 		}
 	}
 
-	if (!SocketHandle)
+	if (!SocketHandle && RemoteDesc.TlsRequirement != ETlsRequirement::Required)
 	{
 		FSocketHandle RawSocketHandle = SocketConnectTcp(RemoteDesc.Host.Address.c_str(), RemoteDesc.Host.Port);
 		SocketSetRecvTimeout(RawSocketHandle, RemoteDesc.RecvTimeoutSeconds);
@@ -408,7 +405,7 @@ TResult<ProxyQuery::FHelloResponse>
 ProxyQuery::Hello(const FRemoteDesc& RemoteDesc, const FAuthDesc* OptAuthDesc)
 {
 	FTlsClientSettings TlsSettings = RemoteDesc.GetTlsClientSettings();
-	FHttpConnection	   Connection(RemoteDesc.Host.Address, RemoteDesc.Host.Port, RemoteDesc.bTlsEnable ? &TlsSettings : nullptr);
+	FHttpConnection	   Connection(RemoteDesc.Host.Address, RemoteDesc.Host.Port, RemoteDesc.TlsRequirement, TlsSettings);
 	return Hello(Connection, OptAuthDesc);
 }
 
@@ -528,6 +525,8 @@ ProxyQuery::Hello(FHttpConnection& HttpConnection, const FAuthDesc* OptAuthDesc)
 		}
 	}
 
+	Result.bConnectionEncrypted = HttpConnection.IsEncrypted();
+
 	return ResultOk(std::move(Result));
 }
 
@@ -622,7 +621,7 @@ ProxyQuery::DownloadFile(const FRemoteDesc&					 Remote,
 	auto CreateConnection = [Remote]
 	{
 		FTlsClientSettings TlsSettings = Remote.GetTlsClientSettings();
-		return new FHttpConnection(Remote.Host.Address, Remote.Host.Port, &TlsSettings);
+		return new FHttpConnection(Remote.Host.Address, Remote.Host.Port, Remote.TlsRequirement, TlsSettings);
 	};
 
 	TObjectPool<FHttpConnection> ConnectionPool(CreateConnection);
@@ -968,11 +967,11 @@ FProxyPool::FProxyPool(const FRemoteDesc& InRemoteDesc, const FAuthDesc* InAuthD
 		else
 		{
 			const ProxyQuery::FHelloResponse& Data = Response.GetData();
-			UNSYNC_VERBOSE(L"Connection established. Server name: %hs, version: %hs, git: %hs, tls: %hs.",
+			UNSYNC_VERBOSE(L"Connection established. Server name: %hs, version: %hs, git: %hs, tls: %hs",
 						   Data.Name.empty() ? "unknown" : Data.Name.c_str(),
 						   Data.VersionNumber.empty() ? "unknown" : Data.VersionNumber.c_str(),
 						   Data.VersionGit.empty() ? "unknown" : Data.VersionGit.c_str(),
-						   RemoteDesc.bTlsEnable ? "yes" : "no"
+						   (Data.bConnectionEncrypted) ? "yes" : "no"
 				);
 
 			Features  = Data.Features;
diff --git a/Engine/Source/Programs/Unsync/Private/UnsyncProxy.h b/Engine/Source/Programs/Unsync/Private/UnsyncProxy.h
--- a/Engine/Source/Programs/Unsync/Private/UnsyncProxy.h
+++ b/Engine/Source/Programs/Unsync/Private/UnsyncProxy.h
@@ -211,6 +211,10 @@ struct FHelloResponse
 
 	std::optional<FHostAddressAndPort> PrimaryHost;
 
+	// Derived data
+
+	bool bConnectionEncrypted = false;
+
 	bool SupportsAuthentication() const { return Features.bAuthentication && !AuthServerUri.empty() && !AuthClientId.empty(); }
 };
 TResult<FHelloResponse> Hello(const FRemoteDesc& RemoteDesc, const FAuthDesc* OptAuthDesc = nullptr);
diff --git a/Engine/Source/Programs/Unsync/Private/UnsyncRemote.cpp b/Engine/Source/Programs/Unsync/Private/UnsyncRemote.cpp
--- a/Engine/Source/Programs/Unsync/Private/UnsyncRemote.cpp
+++ b/Engine/Source/Programs/Unsync/Private/UnsyncRemote.cpp
@@ -83,25 +83,23 @@ FRemoteDesc::FromUrl(std::string_view Url)
 
 	std::string_view Scheme = bHaveScheme ? Url.substr(0, SchemePos) : std::string_view();
 
-	bool bUseTls = false;
-
-	std::string_view HostAddress = bHaveScheme ? Url.substr(SchemePos + 3) : Url;
-
-	ETransportProtocol Transport = ETransportProtocol::Unsync;
+	ETlsRequirement	   TlsRequirement = ETlsRequirement::None;
+	std::string_view   HostAddress	  = bHaveScheme ? Url.substr(SchemePos + 3) : Url;
+	ETransportProtocol Transport	  = ETransportProtocol::Unsync;
 
 	if (Scheme.ends_with("https"))
 	{
-		bUseTls	  = true;
-		Transport = ETransportProtocol::Http;
+		TlsRequirement = ETlsRequirement::Required;
+		Transport	   = ETransportProtocol::Http;
 	}
 	else if (Scheme.ends_with("http"))
 	{
-		bUseTls	  = false;
-		Transport = ETransportProtocol::Http;
+		TlsRequirement = ETlsRequirement::None;
+		Transport	   = ETransportProtocol::Http;
 	}
 	else if (Scheme.ends_with("tls"))
 	{
-		bUseTls = true;
+		TlsRequirement = ETlsRequirement::Required;
 		if (Scheme.starts_with("unsync"))
 		{
 			Transport = ETransportProtocol::Unsync;
@@ -138,8 +136,8 @@ FRemoteDesc::FromUrl(std::string_view Url)
 		HostAddress				= HostAddress.substr(0, NamespacePos);
 	}
 
-	uint16 HostPort = 0;
-	const size_t PortPos	= HostAddress.find_first_of(':');
+	uint16		 HostPort = 0;
+	const size_t PortPos  = HostAddress.find_first_of(':');
 
 	const size_t RequestPos = HostAddress.find_first_of('/');
 	if (RequestPos != std::string::npos)
@@ -157,7 +155,7 @@ FRemoteDesc::FromUrl(std::string_view Url)
 				HostPort = UNSYNC_DEFAULT_PORT;
 				break;
 			case ETransportProtocol::Http:
-				HostPort = bUseTls ? 443 : 80;
+				HostPort = (TlsRequirement == ETlsRequirement::Required) ? 443 : 80;
 				break;
 		}
 	}
@@ -171,22 +169,23 @@ FRemoteDesc::FromUrl(std::string_view Url)
 		HostAddress = HostAddress.substr(0, PortPos);
 	}
 
-	if (HostPort == 443)
-	{
-		bUseTls = true;
-	}
-
 	if (HostPort == 0)
 	{
 		return AppError("Invalid host port");  // TODO: extract the port substring
 	}
 
-	Result.Host.Address = HostAddress;
-	Result.Host.Port	= HostPort;
+	if (HostPort == 443)
+	{
+		TlsRequirement = ETlsRequirement::Required;
+	}
+	else if (Result.Protocol == EProtocolFlavor::Unsync && TlsRequirement < ETlsRequirement::Required)
+	{
+		TlsRequirement = ETlsRequirement::Preferred;
+	}
 
-	Result.bTlsEnable			 = bUseTls;
-	Result.bTlsVerifyCertificate = bUseTls;
-	Result.bTlsVerifySubject	 = bUseTls;
+	Result.Host.Address	  = HostAddress;
+	Result.Host.Port	  = HostPort;
+	Result.TlsRequirement = TlsRequirement;
 
 	return ResultOk(Result);
 }
@@ -195,7 +194,7 @@ FTlsClientSettings
 FRemoteDesc::GetTlsClientSettings() const
 {
 	FTlsClientSettings Result = {};
-	if (bTlsEnable)
+	if (TlsRequirement != ETlsRequirement::None)
 	{
 		Result.Subject			  = GetTlsSubject();
 		Result.bVerifyCertificate = bTlsVerifyCertificate;
@@ -223,7 +222,7 @@ TestParseRemote()
 	{
 		TResult<FRemoteDesc> ParseResult = FRemoteDesc::FromUrl("example.com");
 		UNSYNC_ASSERT(ParseResult.IsOk());
-		UNSYNC_ASSERT(ParseResult->bTlsEnable == false);
+		UNSYNC_ASSERT(ParseResult->TlsRequirement == ETlsRequirement::Preferred);
 		UNSYNC_ASSERT(ParseResult->Host.Address == "example.com");
 		UNSYNC_ASSERT(ParseResult->Host.Port == UNSYNC_DEFAULT_PORT);
 		UNSYNC_ASSERT(ParseResult->Protocol == EProtocolFlavor::Unsync);
@@ -232,7 +231,7 @@ TestParseRemote()
 	{
 		TResult<FRemoteDesc> ParseResult = FRemoteDesc::FromUrl("http://example.com#foo");
 		UNSYNC_ASSERT(ParseResult.IsOk());
-		UNSYNC_ASSERT(ParseResult->bTlsEnable == false);
+		UNSYNC_ASSERT(ParseResult->TlsRequirement == ETlsRequirement::None);
 		UNSYNC_ASSERT(ParseResult->Host.Address == "example.com");
 		UNSYNC_ASSERT(ParseResult->Host.Port == 80);
 		UNSYNC_ASSERT(ParseResult->Protocol == EProtocolFlavor::Jupiter);
@@ -241,7 +240,7 @@ TestParseRemote()
 	{
 		TResult<FRemoteDesc> ParseResult = FRemoteDesc::FromUrl("https://example.com#foo");
 		UNSYNC_ASSERT(ParseResult.IsOk());
-		UNSYNC_ASSERT(ParseResult->bTlsEnable == true);
+		UNSYNC_ASSERT(ParseResult->TlsRequirement == ETlsRequirement::Required);
 		UNSYNC_ASSERT(ParseResult->Host.Address == "example.com");
 		UNSYNC_ASSERT(ParseResult->Host.Port == 443);
 		UNSYNC_ASSERT(ParseResult->Protocol == EProtocolFlavor::Jupiter);
@@ -250,7 +249,7 @@ TestParseRemote()
 	{
 		TResult<FRemoteDesc> ParseResult = FRemoteDesc::FromUrl("http://example.com:1234#foo");
 		UNSYNC_ASSERT(ParseResult.IsOk());
-		UNSYNC_ASSERT(ParseResult->bTlsEnable == false);
+		UNSYNC_ASSERT(ParseResult->TlsRequirement == ETlsRequirement::None);
 		UNSYNC_ASSERT(ParseResult->Host.Address == "example.com");
 		UNSYNC_ASSERT(ParseResult->Host.Port == 1234);
 		UNSYNC_ASSERT(ParseResult->Protocol == EProtocolFlavor::Jupiter);
@@ -259,7 +258,7 @@ TestParseRemote()
 	{
 		TResult<FRemoteDesc> ParseResult = FRemoteDesc::FromUrl("https://example.com:1234#foo");
 		UNSYNC_ASSERT(ParseResult.IsOk());
-		UNSYNC_ASSERT(ParseResult->bTlsEnable == true);
+		UNSYNC_ASSERT(ParseResult->TlsRequirement == ETlsRequirement::Required);
 		UNSYNC_ASSERT(ParseResult->Host.Address == "example.com");
 		UNSYNC_ASSERT(ParseResult->Host.Port == 1234);
 		UNSYNC_ASSERT(ParseResult->Protocol == EProtocolFlavor::Jupiter);
@@ -268,7 +267,7 @@ TestParseRemote()
 	{
 		TResult<FRemoteDesc> ParseResult = FRemoteDesc::FromUrl("jupiter+http://example.com:1234");
 		UNSYNC_ASSERT(ParseResult.IsOk());
-		UNSYNC_ASSERT(ParseResult->bTlsEnable == false);
+		UNSYNC_ASSERT(ParseResult->TlsRequirement == ETlsRequirement::None);
 		UNSYNC_ASSERT(ParseResult->Host.Address == "example.com");
 		UNSYNC_ASSERT(ParseResult->Host.Port == 1234);
 		UNSYNC_ASSERT(ParseResult->Protocol == EProtocolFlavor::Jupiter);
@@ -277,7 +276,7 @@ TestParseRemote()
 	{
 		TResult<FRemoteDesc> ParseResult = FRemoteDesc::FromUrl("jupiter+https://example.com");
 		UNSYNC_ASSERT(ParseResult.IsOk());
-		UNSYNC_ASSERT(ParseResult->bTlsEnable == true);
+		UNSYNC_ASSERT(ParseResult->TlsRequirement == ETlsRequirement::Required);
 		UNSYNC_ASSERT(ParseResult->Host.Address == "example.com");
 		UNSYNC_ASSERT(ParseResult->Host.Port == 443);
 		UNSYNC_ASSERT(ParseResult->Protocol == EProtocolFlavor::Jupiter);
@@ -286,7 +285,7 @@ TestParseRemote()
 	{
 		TResult<FRemoteDesc> ParseResult = FRemoteDesc::FromUrl("jupiter+https://example.com#test.namespace");
 		UNSYNC_ASSERT(ParseResult.IsOk());
-		UNSYNC_ASSERT(ParseResult->bTlsEnable == true);
+		UNSYNC_ASSERT(ParseResult->TlsRequirement == ETlsRequirement::Required);
 		UNSYNC_ASSERT(ParseResult->Host.Address == "example.com");
 		UNSYNC_ASSERT(ParseResult->Host.Port == 443);
 		UNSYNC_ASSERT(ParseResult->Protocol == EProtocolFlavor::Jupiter);
@@ -296,7 +295,7 @@ TestParseRemote()
 	{
 		TResult<FRemoteDesc> ParseResult = FRemoteDesc::FromUrl("jupiter+https://example.com:1234#test.namespace");
 		UNSYNC_ASSERT(ParseResult.IsOk());
-		UNSYNC_ASSERT(ParseResult->bTlsEnable == true);
+		UNSYNC_ASSERT(ParseResult->TlsRequirement == ETlsRequirement::Required);
 		UNSYNC_ASSERT(ParseResult->Host.Address == "example.com");
 		UNSYNC_ASSERT(ParseResult->Host.Port == 1234);
 		UNSYNC_ASSERT(ParseResult->Protocol == EProtocolFlavor::Jupiter);
@@ -306,7 +305,7 @@ TestParseRemote()
 	{
 		TResult<FRemoteDesc> ParseResult = FRemoteDesc::FromUrl("example.com:1234#test.namespace");
 		UNSYNC_ASSERT(ParseResult.IsOk());
-		UNSYNC_ASSERT(ParseResult->bTlsEnable == false);
+		UNSYNC_ASSERT(ParseResult->TlsRequirement == ETlsRequirement::Preferred);
 		UNSYNC_ASSERT(ParseResult->Host.Address == "example.com");
 		UNSYNC_ASSERT(ParseResult->Host.Port == 1234);
 		UNSYNC_ASSERT(ParseResult->Protocol == EProtocolFlavor::Unsync);
@@ -316,7 +315,7 @@ TestParseRemote()
 	{
 		TResult<FRemoteDesc> ParseResult = FRemoteDesc::FromUrl("unsync://example.com:1234");
 		UNSYNC_ASSERT(ParseResult.IsOk());
-		UNSYNC_ASSERT(ParseResult->bTlsEnable == false);
+		UNSYNC_ASSERT(ParseResult->TlsRequirement == ETlsRequirement::Preferred);
 		UNSYNC_ASSERT(ParseResult->Host.Address == "example.com");
 		UNSYNC_ASSERT(ParseResult->Host.Port == 1234);
 		UNSYNC_ASSERT(ParseResult->Protocol == EProtocolFlavor::Unsync);
@@ -325,7 +324,7 @@ TestParseRemote()
 	{
 		TResult<FRemoteDesc> ParseResult = FRemoteDesc::FromUrl("unsync+tls://example.com:1234");
 		UNSYNC_ASSERT(ParseResult.IsOk());
-		UNSYNC_ASSERT(ParseResult->bTlsEnable == true);
+		UNSYNC_ASSERT(ParseResult->TlsRequirement == ETlsRequirement::Required);
 		UNSYNC_ASSERT(ParseResult->Host.Address == "example.com");
 		UNSYNC_ASSERT(ParseResult->Host.Port == 1234);
 		UNSYNC_ASSERT(ParseResult->Protocol == EProtocolFlavor::Unsync);
@@ -339,7 +338,7 @@ TestParseRemote()
 	{
 		TResult<FRemoteDesc> ParseResult = FRemoteDesc::FromUrl("unsync+tls://example.com:1234/request/path#namespace");
 		UNSYNC_ASSERT(ParseResult.IsOk());
-		UNSYNC_ASSERT(ParseResult->bTlsEnable == true);
+		UNSYNC_ASSERT(ParseResult->TlsRequirement == ETlsRequirement::Required);
 		UNSYNC_ASSERT(ParseResult->Host.Address == "example.com");
 		UNSYNC_ASSERT(ParseResult->Host.Port == 1234);
 		UNSYNC_ASSERT(ParseResult->Protocol == EProtocolFlavor::Unsync);
@@ -350,7 +349,7 @@ TestParseRemote()
 	{
 		TResult<FRemoteDesc> ParseResult = FRemoteDesc::FromUrl("http://example.com/request/path#namespace");
 		UNSYNC_ASSERT(ParseResult.IsOk());
-		UNSYNC_ASSERT(ParseResult->bTlsEnable == false);
+		UNSYNC_ASSERT(ParseResult->TlsRequirement == ETlsRequirement::None);
 		UNSYNC_ASSERT(ParseResult->Host.Address == "example.com");
 		UNSYNC_ASSERT(ParseResult->Host.Port == 80);
 		UNSYNC_ASSERT(ParseResult->Protocol == EProtocolFlavor::Jupiter);
@@ -361,7 +360,7 @@ TestParseRemote()
 	{
 		TResult<FRemoteDesc> ParseResult = FRemoteDesc::FromUrl("example.com:443");
 		UNSYNC_ASSERT(ParseResult.IsOk());
-		UNSYNC_ASSERT(ParseResult->bTlsEnable == true);
+		UNSYNC_ASSERT(ParseResult->TlsRequirement == ETlsRequirement::Required);
 		UNSYNC_ASSERT(ParseResult->Host.Address == "example.com");
 		UNSYNC_ASSERT(ParseResult->Host.Port == 443);
 		UNSYNC_ASSERT(ParseResult->Protocol == EProtocolFlavor::Unsync);
diff --git a/Engine/Source/Programs/Unsync/Private/UnsyncRemote.h b/Engine/Source/Programs/Unsync/Private/UnsyncRemote.h
--- a/Engine/Source/Programs/Unsync/Private/UnsyncRemote.h
+++ b/Engine/Source/Programs/Unsync/Private/UnsyncRemote.h
@@ -49,9 +49,9 @@ struct FRemoteDesc
 	std::string StorageBucket = "unsync";  // TODO: override via command line
 	std::string HttpHeaders;
 
-	bool					 bTlsEnable			   = true;	// Prefer TLS, if supported by protocol and remote server
 	bool					 bTlsVerifyCertificate = true;	// Disabling this allows self-signed certificates
 	bool					 bTlsVerifySubject	   = true;	// Disabling this is insecure, but may be useful during development
+	ETlsRequirement			 TlsRequirement		   = ETlsRequirement::None;
 	std::string				 TlsSubjectOverride;			// Use host address if empty (default)
 	std::shared_ptr<FBuffer> TlsCacert;	 // Custom CA to use for server certificate validation (system root CA is used by default)
 
diff --git a/Engine/Source/Programs/Unsync/Private/UnsyncSocket.h b/Engine/Source/Programs/Unsync/Private/UnsyncSocket.h
--- a/Engine/Source/Programs/Unsync/Private/UnsyncSocket.h
+++ b/Engine/Source/Programs/Unsync/Private/UnsyncSocket.h
@@ -89,6 +89,8 @@ struct FSocketBase
 	virtual int RecvAll(void* Data, size_t DataSize)	= 0;
 	virtual int RecvAny(void* Data, size_t DataSize)	= 0;
 
+	virtual bool IsEncrypted() const = 0;
+
 	FSocketHandle	Handle	 = {};
 	ESocketSecurity Security = ESocketSecurity::None;
 };
@@ -100,6 +102,14 @@ struct FSocketRaw : FSocketBase
 	virtual int Send(const void* Data, size_t DataSize) override;
 	virtual int RecvAll(void* Data, size_t DataSize) override;
 	virtual int RecvAny(void* Data, size_t DataSize) override;
+
+	virtual bool IsEncrypted() const override { return false; }
+};
+
+enum class ETlsRequirement : uint8 {
+	None,		// Use defaults for the underlying protocol
+	Preferred,	// Try TLS connection, but allow falling back to plaintext if server supports it
+	Required,	// Only allow connections with TLS
 };
 
 struct FTlsClientSettings
@@ -115,12 +125,14 @@ struct FSocketTls : FSocketBase
 	FSocketTls(FSocketHandle InHandle, FTlsClientSettings ClientSettings);
 	~FSocketTls();
 
-	bool IsTlsValid() { return TlsCtx != nullptr; }
+	bool IsTlsValid() const { return TlsCtx != nullptr; }
 
 	virtual int Send(const void* Data, size_t DataSize) override;
 	virtual int RecvAll(void* Data, size_t DataSize) override;
 	virtual int RecvAny(void* Data, size_t DataSize) override;
 
+	virtual bool IsEncrypted() const override { return IsTlsValid(); }
+
 	tls* TlsCtx = {};
 };
 
