unsync - Fix inconsistent serialized directory separators in the manifest

* Fixed syncing data sets that had manifests generated on non-windows platform

diff --git a/Engine/Source/Programs/Unsync/Private/UnsyncCore.cpp b/Engine/Source/Programs/Unsync/Private/UnsyncCore.cpp
--- a/Engine/Source/Programs/Unsync/Private/UnsyncCore.cpp
+++ b/Engine/Source/Programs/Unsync/Private/UnsyncCore.cpp
@@ -27,7 +27,7 @@ UNSYNC_THIRD_PARTY_INCLUDES_START
 #include <md5-sse2.h>
 UNSYNC_THIRD_PARTY_INCLUDES_END
 
-#define UNSYNC_VERSION_STR "1.0.66"
+#define UNSYNC_VERSION_STR "1.0.67"
 
 namespace unsync {
 
diff --git a/Engine/Source/Programs/Unsync/Private/UnsyncSerialization.cpp b/Engine/Source/Programs/Unsync/Private/UnsyncSerialization.cpp
--- a/Engine/Source/Programs/Unsync/Private/UnsyncSerialization.cpp
+++ b/Engine/Source/Programs/Unsync/Private/UnsyncSerialization.cpp
@@ -395,6 +395,21 @@ LoadFileRevisionControl(FIOReaderStream& Reader, FSerializedSectionHeader Header
 	return Result;
 }
 
+static void
+ConvertSerializedPathSeparatorsToNative(std::string& PathUtf8)
+{
+	std::replace_if(
+		PathUtf8.begin(),
+		PathUtf8.end(),
+		[](char C) { return C == '/' || C == '\\'; },
+		char(FPath::preferred_separator));
+}
+
+static void
+ConvertNativePathSeparatorsToSerialized(std::string& PathUtf8)
+{
+	std::replace(PathUtf8.begin(), PathUtf8.end(), char(FPath::preferred_separator), '/');
+}
 
 bool  // TODO: return a TResult
 LoadDirectoryManifest(FDirectoryManifest& OutManifest, const FPath& Root, FIOReaderStream& Stream)
@@ -528,7 +543,8 @@ LoadDirectoryManifest(FDirectoryManifest& OutManifest, const FPath& Root, FIORea
 		for (uint64 FileIndex = 0; FileIndex < NumFiles; ++FileIndex)
 		{
 			Serialize(Stream, FilenameUtf8);
-
+			ConvertSerializedPathSeparatorsToNative(FilenameUtf8);
+			
 			FFileManifest FileManifest;
 			Serialize(Stream, FileManifest.Mtime);
 			Serialize(Stream, FileManifest.Size);
@@ -575,7 +591,7 @@ LoadDirectoryManifest(FDirectoryManifest& OutManifest, const FPath& Root, FIORea
 				return false;
 			}
 
-			FileManifest.CurrentPath	= Root / Filename;
+			FileManifest.CurrentPath = Root / Filename;
 
 			if (bReadOnlyMaskValid)
 			{
@@ -748,6 +764,8 @@ SaveDirectoryManifest(const FDirectoryManifest& Manifest, FVectorStreamOut& Stre
 	for (const auto& It : Manifest.Files)
 	{
 		std::string FilenameUtf8 = ConvertWideToUtf8(It.first);
+		ConvertNativePathSeparatorsToSerialized(FilenameUtf8);
+
 		Serialize(Stream, FilenameUtf8);  // name as utf8
 
 		const FFileManifest& FileManifest = It.second;
