unsync - Fix downloading uncompressed blocks, tweak direct file download connection limit and chunk size

diff --git a/Engine/Source/Programs/Unsync/Private/UnsyncCore.cpp b/Engine/Source/Programs/Unsync/Private/UnsyncCore.cpp
--- a/Engine/Source/Programs/Unsync/Private/UnsyncCore.cpp
+++ b/Engine/Source/Programs/Unsync/Private/UnsyncCore.cpp
@@ -2060,7 +2060,7 @@ struct FFileSyncTaskBatch
 
 						bool bOk = true;
 
-						if (Block.IsCompressed())
+						if (Block.bCompressed)
 						{
 							bOk = Decompress(Block.Data, Block.CompressedSize, OutputView.Data, OutputView.Size);
 						}
diff --git a/Engine/Source/Programs/Unsync/Private/UnsyncJupiter.cpp b/Engine/Source/Programs/Unsync/Private/UnsyncJupiter.cpp
--- a/Engine/Source/Programs/Unsync/Private/UnsyncJupiter.cpp
+++ b/Engine/Source/Programs/Unsync/Private/UnsyncJupiter.cpp
@@ -939,7 +939,8 @@ FJupiterProtocolImpl::Download(const TArrayView<FNeedBlock> NeedBlocks, const FB
 			UNSYNC_ASSERT(OffsetFromFirst + Block.Size <= Buffer->Size());
 
 			FDownloadedBlock ResultBlock;
-			ResultBlock.CompressedSize	 = 0;  // TODO: support compression
+			ResultBlock.bCompressed		 = false;
+			ResultBlock.CompressedSize	 = Block.Size;
 			ResultBlock.DecompressedSize = Block.Size;
 			ResultBlock.Data			 = Buffer->Data() + OffsetFromFirst;
 
diff --git a/Engine/Source/Programs/Unsync/Private/UnsyncProtocol.h b/Engine/Source/Programs/Unsync/Private/UnsyncProtocol.h
--- a/Engine/Source/Programs/Unsync/Private/UnsyncProtocol.h
+++ b/Engine/Source/Programs/Unsync/Private/UnsyncProtocol.h
@@ -178,8 +178,8 @@ struct FRequestBlocksPacket
 struct FBlockPacket
 {
 	FHash128 Hash			  = {};
-	uint64	 DecompressedSize = 0;
-	FBuffer	 CompressedData;
+	uint64	 DecompressedSize = 0; // 0 if data is not compressed
+	FBuffer	 Data;
 };
 
 struct FPatchHeader
diff --git a/Engine/Source/Programs/Unsync/Private/UnsyncProxy.cpp b/Engine/Source/Programs/Unsync/Private/UnsyncProxy.cpp
--- a/Engine/Source/Programs/Unsync/Private/UnsyncProxy.cpp
+++ b/Engine/Source/Programs/Unsync/Private/UnsyncProxy.cpp
@@ -342,9 +342,9 @@ FUnsyncProtocolImpl::Download(const TArrayView<FNeedBlock> NeedBlocks, const FBl
 			break;
 		}
 
-		BlockPacket.CompressedData.Resize(CompressedDataSize);
+		BlockPacket.Data.Resize(CompressedDataSize);
 
-		bOk &= (SocketRecvAll(*SocketHandle, BlockPacket.CompressedData.Data(), BlockPacket.CompressedData.Size()) == CompressedDataSize);
+		bOk &= (SocketRecvAll(*SocketHandle, BlockPacket.Data.Data(), BlockPacket.Data.Size()) == CompressedDataSize);
 
 		if (!bOk)
 		{
@@ -363,9 +363,21 @@ FUnsyncProtocolImpl::Download(const TArrayView<FNeedBlock> NeedBlocks, const FBl
 		if (bOk)
 		{
 			FDownloadedBlock DownloadedBlock;
-			DownloadedBlock.DecompressedSize = BlockPacket.DecompressedSize;
-			DownloadedBlock.CompressedSize	 = BlockPacket.CompressedData.Size();
-			DownloadedBlock.Data			 = BlockPacket.CompressedData.Data();
+
+			DownloadedBlock.CompressedSize = BlockPacket.Data.Size();
+			DownloadedBlock.Data		   = BlockPacket.Data.Data();
+
+			if (BlockPacket.DecompressedSize != 0)
+			{
+				DownloadedBlock.bCompressed		 = true;
+				DownloadedBlock.DecompressedSize = BlockPacket.DecompressedSize;
+			}
+			else
+			{
+				DownloadedBlock.bCompressed		 = false;
+				DownloadedBlock.DecompressedSize = DownloadedBlock.CompressedSize;
+			}
+
 			CompletionCallback(DownloadedBlock, BlockPacket.Hash);
 		}
 
@@ -660,7 +672,7 @@ ProxyQuery::DownloadFile(const FRemoteDesc&					 Remote,
 		ConnectionPool.Release(std::move(Connection));
 	}
 
-	const uint64 MaxChunkSize = 16_MB;
+	const uint64 MaxChunkSize = 8_MB;
 
 	FIOWriter& Result = OutputCallback(FileSize);
 	if (!Result.IsValid())
@@ -681,7 +693,7 @@ ProxyQuery::DownloadFile(const FRemoteDesc&					 Remote,
 	}
 
 	FAtomicError Error;
-	FSemaphore	 DownloadSempahore(4);	// up to 4 concurrent connections
+	FSemaphore	 DownloadSempahore(8);	// limit concurrent connections
 
 	FLogProgressScope DownloadProgress(FileSize, ELogProgressUnits::MB);
 
diff --git a/Engine/Source/Programs/Unsync/Private/UnsyncProxy.h b/Engine/Source/Programs/Unsync/Private/UnsyncProxy.h
--- a/Engine/Source/Programs/Unsync/Private/UnsyncProxy.h
+++ b/Engine/Source/Programs/Unsync/Private/UnsyncProxy.h
@@ -43,8 +43,7 @@ struct FDownloadedBlock
 	uint64		 DecompressedSize = 0;
 	uint64		 CompressedSize	  = 0;
 	const uint8* Data			  = nullptr;
-
-	bool IsCompressed() const { return CompressedSize != 0; }
+	bool		 bCompressed	  = false;
 };
 
 using FBlockDownloadCallback = std::function<void(const FDownloadedBlock& Block, FHash128 BlockHash)>;
diff --git a/Engine/Source/Programs/Unsync/Private/UnsyncTarget.cpp b/Engine/Source/Programs/Unsync/Private/UnsyncTarget.cpp
--- a/Engine/Source/Programs/Unsync/Private/UnsyncTarget.cpp
+++ b/Engine/Source/Programs/Unsync/Private/UnsyncTarget.cpp
@@ -523,9 +523,10 @@ BuildTarget(FIOWriter& Output, FIOReader& Source, FIOReader& Base, const FNeedLi
 					FBlockWriteCmd Cmd = WriteIt->second;
 
 					// TODO: avoid this copy by storing IOBuffer in DownloadedBlock
-					bool	  bCompressed	 = Block.IsCompressed();
-					uint64	  DownloadedSize = bCompressed ? Block.CompressedSize : Block.DecompressedSize;
-					FIOBuffer DownloadedData = FIOBuffer::Alloc(DownloadedSize, L"downloaded_data");
+					const bool	 bCompressed	= Block.bCompressed;
+					const uint64 DownloadedSize = bCompressed ? Block.CompressedSize : Block.DecompressedSize;
+					FIOBuffer	 DownloadedData = FIOBuffer::Alloc(DownloadedSize, L"downloaded_data");
+
 					memcpy(DownloadedData.GetData(), Block.Data, DownloadedSize);
 
 					DecompressionSemaphore.Acquire();
