unsync - ListDirectory and DownloadFile helpers now take FHttpConnection

diff --git a/Engine/Source/Programs/Unsync/Private/UnsyncCmdQuery.cpp b/Engine/Source/Programs/Unsync/Private/UnsyncCmdQuery.cpp
--- a/Engine/Source/Programs/Unsync/Private/UnsyncCmdQuery.cpp
+++ b/Engine/Source/Programs/Unsync/Private/UnsyncCmdQuery.cpp
@@ -474,7 +474,8 @@ CmdQueryFile(const FCmdQueryOptions& Options)
 		return *ResultWriter;
 	};
 
-	TResult<> Response = ProxyQuery::DownloadFile(Options.Remote, &AuthDesc, Options.Args[0], OutputCallback);
+	FHttpConnection Connection = FHttpConnection::CreateDefaultHttps(Options.Remote);
+	TResult<> Response = ProxyQuery::DownloadFile(Connection, &AuthDesc, Options.Args[0], OutputCallback);
 
 	if (Response.IsOk())
 	{
diff --git a/Engine/Source/Programs/Unsync/Private/UnsyncCore.cpp b/Engine/Source/Programs/Unsync/Private/UnsyncCore.cpp
--- a/Engine/Source/Programs/Unsync/Private/UnsyncCore.cpp
+++ b/Engine/Source/Programs/Unsync/Private/UnsyncCore.cpp
@@ -14,6 +14,7 @@
 #include "UnsyncThread.h"
 #include "UnsyncUtil.h"
 #include "UnsyncTarget.h"
+#include "UnsyncHttp.h"
 
 #include <condition_variable>
 #include <filesystem>
@@ -1849,7 +1850,7 @@ struct FPooledProxy
 };
 
 static bool
-DownloadFileIfNewer(const FRemoteDesc& RemoteDesc, const FAuthDesc* AuthDesc, const FPath& Source, const FPath& Target, EFileMode TargetFileMode)
+DownloadFileIfNewer(FHttpConnection& Connection, const FAuthDesc* AuthDesc, const FPath& Source, const FPath& Target, EFileMode TargetFileMode)
 {
 	using FDirectoryListing		 = ProxyQuery::FDirectoryListing;
 	using FDirectoryListingEntry = ProxyQuery::FDirectoryListingEntry;
@@ -1862,7 +1863,7 @@ DownloadFileIfNewer(const FRemoteDesc& RemoteDesc, const FAuthDesc* AuthDesc, co
 	std::string SourceFileNameUtf8 = ConvertWideToUtf8(SourceFileName.wstring());
 
 	// TODO: could have a dedicated single file stat query
-	TResult<FDirectoryListing> DirectoryListingResult = ProxyQuery::ListDirectory(RemoteDesc, AuthDesc, SourceParentUtf8);
+	TResult<FDirectoryListing> DirectoryListingResult = ProxyQuery::ListDirectory(Connection, AuthDesc, SourceParentUtf8);
 	if (DirectoryListingResult.IsError())
 	{
 		LogError(DirectoryListingResult.GetError());
@@ -1890,7 +1891,7 @@ DownloadFileIfNewer(const FRemoteDesc& RemoteDesc, const FAuthDesc* AuthDesc, co
 	if (SourceEntry->Size != TargetAttr.Size || SourceEntry->Mtime != TargetAttr.Mtime)
 	{
 		UNSYNC_VERBOSE(L"Downloading '%ls'", Source.wstring().c_str());
-		TResult<FBuffer> DownloadResult = ProxyQuery::DownloadFile(RemoteDesc, AuthDesc, SourceUtf8);
+		TResult<FBuffer> DownloadResult = ProxyQuery::DownloadFile(Connection, AuthDesc, SourceUtf8);
 		if (DownloadResult.IsError())
 		{
 			LogError(DownloadResult.GetError());
@@ -1967,7 +1968,10 @@ LoadAndMergeSourceManifest(FDirectoryManifest& Output,
 	if (bDownloadManifestFromProxy)
 	{
 		UNSYNC_LOG_INDENT;
-		bool bDownloadedOk = DownloadFileIfNewer(ProxyPool.RemoteDesc,
+
+		FHttpConnection Connection = FHttpConnection::CreateDefaultHttps(ProxyPool.RemoteDesc);
+
+		bool bDownloadedOk = DownloadFileIfNewer(Connection,
 												 ProxyPool.AuthDesc,
 												 SourceManifestPath,
 												 SourceManifestTempPath,
diff --git a/Engine/Source/Programs/Unsync/Private/UnsyncHttp.cpp b/Engine/Source/Programs/Unsync/Private/UnsyncHttp.cpp
--- a/Engine/Source/Programs/Unsync/Private/UnsyncHttp.cpp
+++ b/Engine/Source/Programs/Unsync/Private/UnsyncHttp.cpp
@@ -281,6 +281,25 @@ HttpRequest(const FRemoteDesc& RemoteDesc, EHttpMethod Method, std::string_view
 	return HttpRequest(RemoteDesc, Method, RequestUrl, EHttpContentType::Unknown, /*Payload*/ {}, BearerToken);
 }
 
+static const char*
+ToString(EHttpMethod Method)
+{
+	switch (Method)
+	{
+		default:
+			UNSYNC_FATAL(L"Unexpected HTTP method %d", (int)Method);
+			return "INVALID";
+		case EHttpMethod::GET:
+			return "GET";
+		case EHttpMethod::HEAD:
+			return "HEAD";
+		case EHttpMethod::POST:
+			return "POST";
+		case EHttpMethod::PUT:
+			return "PUT";
+	}
+}
+
 bool
 HttpRequestBegin(FHttpConnection& Connection, const FHttpRequest& Request)
 {
@@ -300,25 +319,8 @@ HttpRequestBegin(FHttpConnection& Connection, const FHttpRequest& Request)
 
 	// TODO: use a string builder
 	std::string HttpHeader;
-	switch (Request.Method)
-	{
-		default:
-			UNSYNC_FATAL(L"Unexpected HTTP method %d", (int)Request.Method);
-			return false;
-		case EHttpMethod::GET:
-			HttpHeader = "GET ";
-			break;
-		case EHttpMethod::HEAD:
-			HttpHeader = "HEAD ";
-			break;
-		case EHttpMethod::POST:
-			HttpHeader = "POST ";
-			break;
-		case EHttpMethod::PUT:
-			HttpHeader = "PUT ";
-			break;
-	}
-
+	HttpHeader.append(ToString(Request.Method));
+	HttpHeader.append(" ");
 	HttpHeader.append(Request.Url);
 	HttpHeader.append(" HTTP/1.1\r\n");
 
@@ -419,6 +421,8 @@ HttpRequestBegin(FHttpConnection& Connection, const FHttpRequest& Request)
 
 	Connection.NumActiveRequests += 1;
 
+	UNSYNC_VERBOSE2(L"HTTP %hs %.*hs", ToString(Request.Method), int32(Request.Url.length()), Request.Url.data());
+
 	// TODO: detect and handle errors
 	SentBytes += SocketSend(Connection.GetSocket(), HttpHeader.c_str(), HttpHeader.length());
 	if (Request.Payload.Size)
diff --git a/Engine/Source/Programs/Unsync/Private/UnsyncProxy.cpp b/Engine/Source/Programs/Unsync/Private/UnsyncProxy.cpp
--- a/Engine/Source/Programs/Unsync/Private/UnsyncProxy.cpp
+++ b/Engine/Source/Programs/Unsync/Private/UnsyncProxy.cpp
@@ -589,7 +589,7 @@ ProxyQuery::FDirectoryListing::FromJson(const char* JsonString)
 }
 
 TResult<ProxyQuery::FDirectoryListing>
-ProxyQuery::ListDirectory(const FRemoteDesc& Remote, const FAuthDesc* AuthDesc, const std::string& Path)
+ProxyQuery::ListDirectory(FHttpConnection& Connection, const FAuthDesc* AuthDesc, const std::string& Path)
 {
 	std::string Url = fmt::format("/api/v1/list?{}", Path);
 
@@ -605,8 +605,6 @@ ProxyQuery::ListDirectory(const FRemoteDesc& Remote, const FAuthDesc* AuthDesc,
 		BearerToken = std::move(AuthToken->Access);
 	}
 
-	FHttpConnection Connection = FHttpConnection::CreateDefaultHttps(Remote);
-
 	FHttpRequest Request;
 	Request.Url			= Url;
 	Request.Method		= EHttpMethod::GET;
@@ -625,18 +623,16 @@ ProxyQuery::ListDirectory(const FRemoteDesc& Remote, const FAuthDesc* AuthDesc,
 }
 
 TResult<>
-ProxyQuery::DownloadFile(const FRemoteDesc&					 Remote,
+ProxyQuery::DownloadFile(FHttpConnection&					 InConnection,
 						 const FAuthDesc*					 AuthDesc,
 						 const std::string&					 Path,
 						 ProxyQuery::FDownloadOutputCallback OutputCallback)
 {
-	auto CreateConnection = [Remote]
-	{
-		FTlsClientSettings TlsSettings = Remote.GetTlsClientSettings();
-		return new FHttpConnection(Remote.Host.Address, Remote.Host.Port, Remote.TlsRequirement, TlsSettings);
-	};
-
-	TObjectPool<FHttpConnection> ConnectionPool(CreateConnection);
+	TObjectPool<FHttpConnection> ConnectionPool(
+		[&InConnection]
+		{
+			return new FHttpConnection(InConnection);  // Clone the connection
+		});
 
 	std::string Url = fmt::format("/api/v1/file?{}", Path);
 
@@ -710,7 +706,7 @@ ProxyQuery::DownloadFile(const FRemoteDesc&					 Remote,
 	}
 
 	auto ProcessChunk =
-		[&Error, &Result, &Url, &Remote, &ConnectionPool, &DownloadSempahore, &DownloadProgress, &BearerToken](
+		[&Error, &Result, &Url, &ConnectionPool, &DownloadSempahore, &DownloadProgress, &BearerToken](
 			const FRange& Range)
 	{
 		FLogIndentScope	   IndentScope(DownloadProgress.ParentThreadIndent, true);
@@ -767,7 +763,7 @@ ProxyQuery::DownloadFile(const FRemoteDesc&					 Remote,
 }
 
 TResult<FBuffer>
-ProxyQuery::DownloadFile(const FRemoteDesc& Remote, const FAuthDesc* AuthDesc, const std::string& Path)
+ProxyQuery::DownloadFile(FHttpConnection& Connection, const FAuthDesc* AuthDesc, const std::string& Path)
 {
 	FBuffer Result;
 
@@ -780,7 +776,7 @@ ProxyQuery::DownloadFile(const FRemoteDesc& Remote, const FAuthDesc* AuthDesc, c
 		return *ResultWriter;
 	};
 
-	TResult<> DownloadResult = DownloadFile(Remote, AuthDesc, Path, OutputCallback);
+	TResult<> DownloadResult = DownloadFile(Connection, AuthDesc, Path, OutputCallback);
 	if (DownloadResult.IsOk())
 	{
 		return ResultOk(std::move(Result));
diff --git a/Engine/Source/Programs/Unsync/Private/UnsyncProxy.h b/Engine/Source/Programs/Unsync/Private/UnsyncProxy.h
--- a/Engine/Source/Programs/Unsync/Private/UnsyncProxy.h
+++ b/Engine/Source/Programs/Unsync/Private/UnsyncProxy.h
@@ -234,12 +234,14 @@ struct FDirectoryListing
 	static TResult<FDirectoryListing> FromJson(const char* JsonString);
 };
 
-// TODO: add overloads with HTTP connection
-TResult<FDirectoryListing> ListDirectory(const FRemoteDesc& Remote, const FAuthDesc* AuthDesc, const std::string& Path);
-TResult<FBuffer>		   DownloadFile(const FRemoteDesc& Remote, const FAuthDesc* AuthDesc, const std::string& Path);
+TResult<FDirectoryListing> ListDirectory(FHttpConnection& Connection, const FAuthDesc* AuthDesc, const std::string& Path);
+TResult<FBuffer>		   DownloadFile(FHttpConnection& Connection, const FAuthDesc* AuthDesc, const std::string& Path);
 
 using FDownloadOutputCallback = std::function<FIOWriter&(uint64 Size)>;
-TResult<> DownloadFile(const FRemoteDesc& Remote, const FAuthDesc* AuthDesc, const std::string& Path, FDownloadOutputCallback OutputCallback);
+TResult<> DownloadFile(FHttpConnection&		   Connection,
+					   const FAuthDesc*		   AuthDesc,
+					   const std::string&	   Path,
+					   FDownloadOutputCallback OutputCallback);
 
 } 
 
