unsync - Search for .unsyncmanifest as an alternative for .unsync/manifest.bin when downloading manifests from remote server

diff --git a/Engine/Source/Programs/Unsync/Private/UnsyncCore.cpp b/Engine/Source/Programs/Unsync/Private/UnsyncCore.cpp
--- a/Engine/Source/Programs/Unsync/Private/UnsyncCore.cpp
+++ b/Engine/Source/Programs/Unsync/Private/UnsyncCore.cpp
@@ -22,13 +22,14 @@
 #include <string>
 #include <unordered_map>
 #include <unordered_set>
+#include <optional>
 
 UNSYNC_THIRD_PARTY_INCLUDES_START
 #include <blake3.h>
 #include <md5-sse2.h>
 UNSYNC_THIRD_PARTY_INCLUDES_END
 
-#define UNSYNC_VERSION_STR "1.0.67"
+#define UNSYNC_VERSION_STR "1.0.68"
 
 namespace unsync {
 
@@ -1849,48 +1850,82 @@ struct FPooledProxy
 	std::unique_ptr<FProxy> Proxy;
 };
 
-static bool
-DownloadFileIfNewer(FHttpConnection& Connection, const FAuthDesc* AuthDesc, const FPath& Source, const FPath& Target, EFileMode TargetFileMode)
+struct FRemoteManifestInfo
+{
+	FPath							   Path;
+	ProxyQuery::FDirectoryListingEntry Entry;
+};
+
+static TResult<FRemoteManifestInfo>
+FindRemoteUnsyncManifest(FHttpConnection& Connection, const FAuthDesc* AuthDesc, const FPath& RootDirectory)
 {
 	using FDirectoryListing		 = ProxyQuery::FDirectoryListing;
 	using FDirectoryListingEntry = ProxyQuery::FDirectoryListingEntry;
 
-	FPath SourceParent = Source.parent_path();
-	FPath SourceFileName = Source.filename().string();
+	std::string RootDirectoryUtf8 = ConvertWideToUtf8(RootDirectory.wstring());
+	TResult<FDirectoryListing> RootDirectoryListingResult = ProxyQuery::ListDirectory(Connection, AuthDesc, RootDirectoryUtf8);
+	UNSYNC_RETURN_ON_ERROR(RootDirectoryListingResult);
 
-	std::string SourceUtf8		   = ConvertWideToUtf8(Source.wstring());
-	std::string SourceParentUtf8   = ConvertWideToUtf8(SourceParent.wstring());
-	std::string SourceFileNameUtf8 = ConvertWideToUtf8(SourceFileName.wstring());
+	std::optional<FDirectoryListingEntry> FoundUnsyncSubdirectory;
+	std::optional<FDirectoryListingEntry> FoundUnsyncManifestDotfile;
 
-	// TODO: could have a dedicated single file stat query
-	TResult<FDirectoryListing> DirectoryListingResult = ProxyQuery::ListDirectory(Connection, AuthDesc, SourceParentUtf8);
-	if (DirectoryListingResult.IsError())
+	for (const FDirectoryListingEntry& Entry : RootDirectoryListingResult.GetData().Entries)
 	{
-		LogError(DirectoryListingResult.GetError());
-		return false;
-	}
+		if (Entry.Name == ".unsyncmanifest" && !Entry.bDirectory)
+		{
+			FoundUnsyncManifestDotfile = Entry;
+		}
 
-	const FDirectoryListing&	  DirectoryListing = DirectoryListingResult.GetData();
-	const FDirectoryListingEntry* SourceEntry	   = nullptr;
-	for (const FDirectoryListingEntry& Entry : DirectoryListing.Entries)
-	{
-		if (Entry.Name == SourceFileNameUtf8 && !Entry.bDirectory)
+		if (Entry.Name == ".unsync" && Entry.bDirectory)
 		{
-			SourceEntry = &Entry;
-			break;
+			FoundUnsyncSubdirectory = Entry;
 		}
 	}
 
-	if (!SourceEntry)
+	if (FoundUnsyncManifestDotfile && !FoundUnsyncSubdirectory)
 	{
-		UNSYNC_ERROR(L"Remote file '%ls' does not exist", Source.wstring().c_str());
-		return false;
+		FRemoteManifestInfo Result;
+		Result.Path	 = RootDirectory / FoundUnsyncManifestDotfile->Name;
+		Result.Entry = FoundUnsyncManifestDotfile.value();
+		return ResultOk(Result);
 	}
 
+	if (FoundUnsyncSubdirectory)
+	{
+		FPath UnsyncDirectory = RootDirectory / ".unsync";
+		std::string UnsyncDirectoryUtf8 = ConvertWideToUtf8(UnsyncDirectory.wstring());
+
+		TResult<FDirectoryListing> UnsyncDirectoryListingResult = ProxyQuery::ListDirectory(Connection, AuthDesc, UnsyncDirectoryUtf8);
+		UNSYNC_RETURN_ON_ERROR(UnsyncDirectoryListingResult);
+
+		for (const FDirectoryListingEntry& Entry : UnsyncDirectoryListingResult.GetData().Entries)
+		{
+			if (Entry.Name == "manifest.bin" && !Entry.bDirectory)
+			{
+				FRemoteManifestInfo Result;
+				Result.Path	 = UnsyncDirectory / Entry.Name;
+				Result.Entry = FoundUnsyncManifestDotfile.value();
+				return ResultOk(Result);
+			}
+		}
+	}
+
+	return AppError("Could not find remote unsync manifest file");
+}
+
+static bool
+DownloadFileIfPossiblyDifferent(FHttpConnection&		   Connection,
+								const FAuthDesc*		   AuthDesc,
+								const FRemoteManifestInfo& Source,
+								const FPath&			   Target,
+								EFileMode				   TargetFileMode)
+{
 	FFileAttributes TargetAttr = GetFileAttrib(Target);
-	if (SourceEntry->Size != TargetAttr.Size || SourceEntry->Mtime != TargetAttr.Mtime)
+	if (Source.Entry.Size != TargetAttr.Size || Source.Entry.Mtime != TargetAttr.Mtime)
 	{
-		UNSYNC_VERBOSE(L"Downloading '%ls'", Source.wstring().c_str());
+		UNSYNC_VERBOSE(L"Downloading '%ls'", Source.Path.wstring().c_str());
+		std::string SourceUtf8 = ConvertWideToUtf8(Source.Path.wstring());
+
 		TResult<FBuffer> DownloadResult = ProxyQuery::DownloadFile(Connection, AuthDesc, SourceUtf8);
 		if (DownloadResult.IsError())
 		{
@@ -1899,11 +1934,9 @@ DownloadFileIfNewer(FHttpConnection& Connection, const FAuthDesc* AuthDesc, cons
 		}
 
 		const FBuffer& FileBuffer = DownloadResult.GetData();
-		if (FileBuffer.Size() != SourceEntry->Size)
+		if (FileBuffer.Size() != Source.Entry.Size)
 		{
-			UNSYNC_ERROR(L"Downloaded file size mismatch. Expected %llu, actual %llu.",
-						 llu(SourceEntry->Size),
-						 llu(FileBuffer.Size()));
+			UNSYNC_ERROR(L"Downloaded file size mismatch. Expected %llu, actual %llu.", llu(Source.Entry.Size), llu(FileBuffer.Size()));
 			return false;
 		}
 
@@ -1914,7 +1947,7 @@ DownloadFileIfNewer(FHttpConnection& Connection, const FAuthDesc* AuthDesc, cons
 			return false;
 		}
 
-		SetFileMtime(Target, SourceEntry->Mtime, /*allow in dry run*/ true);
+		SetFileMtime(Target, Source.Entry.Mtime, /*allow in dry run*/ true);
 	}
 
 	return true;
@@ -1943,8 +1976,7 @@ LoadAndMergeSourceManifest(FDirectoryManifest& Output,
 
 	FDirectoryManifest LoadedManifest;
 
-	FPath SourceManifestRoot = SourcePath / ".unsync";
-	FPath SourceManifestPath = SourceManifestRoot / "manifest.bin";
+	FPath SourceManifestPath = SourcePath / ".unsync" / "manifest.bin";
 
 	SourceManifestPath = ResolvePath(SourceManifestPath);
 
@@ -1971,11 +2003,22 @@ LoadAndMergeSourceManifest(FDirectoryManifest& Output,
 
 		FHttpConnection Connection = FHttpConnection::CreateDefaultHttps(ProxyPool.RemoteDesc);
 
-		bool bDownloadedOk = DownloadFileIfNewer(Connection,
-												 ProxyPool.AuthDesc,
-												 SourceManifestPath,
-												 SourceManifestTempPath,
-												 EFileMode::CreateWriteOnly | EFileMode::IgnoreDryRun);
+		TResult<FRemoteManifestInfo> FoundManifestResult = FindRemoteUnsyncManifest(Connection, ProxyPool.AuthDesc, SourcePath);
+		UNSYNC_UNUSED(FoundManifestResult);
+
+		if (FoundManifestResult.IsError())
+		{
+			LogError(FoundManifestResult.GetError());
+			return false;
+		}
+
+		const FRemoteManifestInfo& FoundManifest = FoundManifestResult.GetData();
+
+		bool bDownloadedOk = DownloadFileIfPossiblyDifferent(Connection,
+															 ProxyPool.AuthDesc,
+															 FoundManifest,
+															 SourceManifestTempPath,
+															 EFileMode::CreateWriteOnly | EFileMode::IgnoreDryRun);
 
 		if (!bDownloadedOk)
 		{
diff --git a/Engine/Source/Programs/Unsync/Private/UnsyncError.h b/Engine/Source/Programs/Unsync/Private/UnsyncError.h
--- a/Engine/Source/Programs/Unsync/Private/UnsyncError.h
+++ b/Engine/Source/Programs/Unsync/Private/UnsyncError.h
@@ -173,4 +173,12 @@ struct FAtomicError
 	}
 };
 
+#define UNSYNC_RETURN_ON_ERROR(x)           \
+	{                                       \
+		if (x.IsError())                    \
+		{                                   \
+			return std::move(x.GetError()); \
+		}                                   \
+	}
+
 }  // namespace unsync
