unsync - Fix manifest search, output HTTP status codes on query errors

diff --git a/Engine/Source/Programs/Unsync/Private/UnsyncCmdQuery.cpp b/Engine/Source/Programs/Unsync/Private/UnsyncCmdQuery.cpp
--- a/Engine/Source/Programs/Unsync/Private/UnsyncCmdQuery.cpp
+++ b/Engine/Source/Programs/Unsync/Private/UnsyncCmdQuery.cpp
@@ -201,6 +201,12 @@ CmdQueryList(const FCmdQueryOptions& Options)
 	Request.BearerToken	   = AuthToken->Access;
 	FHttpResponse Response = HttpRequest(Connection, Request);
 
+	if (!Response.Success())
+	{
+		LogError(HttpError(Response.Code));
+		return -1;
+	}
+
 	Response.Buffer.PushBack(0);
 
 	std::string	 JsonErrorString;
@@ -337,7 +343,7 @@ CmdQuerySearch(const FCmdQueryOptions& Options)
 
 		if (!Response.Success())
 		{
-			// TODO: report warning
+			LogError(HttpError(Response.Code));
 			return;
 		}
 
@@ -346,7 +352,7 @@ CmdQuerySearch(const FCmdQueryOptions& Options)
 		TResult<FDirectoryListing> DirectoryListingResult = FDirectoryListing::FromJson((const char*)Response.Buffer.Data());
 		if (DirectoryListingResult.IsError())
 		{
-			// TODO: report warning
+			LogError(DirectoryListingResult.GetError());
 			return;
 		}
 
diff --git a/Engine/Source/Programs/Unsync/Private/UnsyncCore.cpp b/Engine/Source/Programs/Unsync/Private/UnsyncCore.cpp
--- a/Engine/Source/Programs/Unsync/Private/UnsyncCore.cpp
+++ b/Engine/Source/Programs/Unsync/Private/UnsyncCore.cpp
@@ -1904,7 +1904,7 @@ FindRemoteUnsyncManifest(FHttpConnection& Connection, const FAuthDesc* AuthDesc,
 			{
 				FRemoteManifestInfo Result;
 				Result.Path	 = UnsyncDirectory / Entry.Name;
-				Result.Entry = FoundUnsyncManifestDotfile.value();
+				Result.Entry = Entry;
 				return ResultOk(Result);
 			}
 		}
diff --git a/Engine/Source/Programs/Unsync/Private/UnsyncUtil.cpp b/Engine/Source/Programs/Unsync/Private/UnsyncUtil.cpp
--- a/Engine/Source/Programs/Unsync/Private/UnsyncUtil.cpp
+++ b/Engine/Source/Programs/Unsync/Private/UnsyncUtil.cpp
@@ -416,14 +416,14 @@ NormalizeFilenameUtf8(const std::string& InFilename)
 
 	FPath FilenameAsPath = ConvertUtf8ToWide(Filename);
 
-	FPath NormalPath = FilenameAsPath.lexically_normal();
 	FPath AbsoluteNormalPath;
 	if (Filename.starts_with("\\\\") || Filename.starts_with("//"))
 	{
-		AbsoluteNormalPath = NormalPath;  // Assume network paths are absolute
+		AbsoluteNormalPath = FilenameAsPath; // Assume network paths are absolute
 	}
 	else
 	{
+		FPath NormalPath	= FilenameAsPath.lexically_normal();
 		FPath CanonicalPath = std::filesystem::weakly_canonical(NormalPath);
 		AbsoluteNormalPath	= std::filesystem::absolute(CanonicalPath);
 	}
