unsync - Groundwork for a cross-platform scheduler

diff --git a/Engine/Source/Programs/Unsync/CMakeLists.txt b/Engine/Source/Programs/Unsync/CMakeLists.txt
--- a/Engine/Source/Programs/Unsync/CMakeLists.txt
+++ b/Engine/Source/Programs/Unsync/CMakeLists.txt
@@ -90,6 +90,8 @@ set(src
 	Private/UnsyncScan.h
 	Private/UnsyncScavenger.cpp
 	Private/UnsyncScavenger.h
+	Private/UnsyncScheduler.cpp
+	Private/UnsyncScheduler.h
 	Private/UnsyncSerialization.cpp
 	Private/UnsyncSerialization.h
 	Private/UnsyncSocket.cpp
diff --git a/Engine/Source/Programs/Unsync/Private/UnsyncCore.cpp b/Engine/Source/Programs/Unsync/Private/UnsyncCore.cpp
--- a/Engine/Source/Programs/Unsync/Private/UnsyncCore.cpp
+++ b/Engine/Source/Programs/Unsync/Private/UnsyncCore.cpp
@@ -15,6 +15,7 @@
 #include "UnsyncUtil.h"
 #include "UnsyncTarget.h"
 #include "UnsyncHttp.h"
+#include "UnsyncScheduler.h"
 
 #include <condition_variable>
 #include <filesystem>
@@ -62,7 +63,6 @@ BlockingReadLarge(FIOReader& Reader, uint64 Offset, uint64 Size, uint8* OutputBu
 	}
 
 	FTaskGroup CopyTasks;
-	FSemaphore IoSemaphore(MAX_ACTIVE_READERS);
 
 	uint64 NumReads = DivUp(ClampedSize, BytesPerRead);
 	for (uint64 ReadIndex = 0; ReadIndex < NumReads; ++ReadIndex)
@@ -71,19 +71,19 @@ BlockingReadLarge(FIOReader& Reader, uint64 Offset, uint64 Size, uint8* OutputBu
 		const uint64 OutputOffset	= BytesPerRead * ReadIndex;
 		const uint64 ThisReadOffset = Offset + OutputOffset;
 
-		IoSemaphore.Acquire();
+		GScheduler.FilesystemSemaphore.Acquire();
 
-		auto ReadCallback = [OutputBuffer, OutputBufferSize, &TotalReadSize, &CopyTasks, &IoSemaphore](FIOBuffer CmdBuffer,
+		auto ReadCallback = [OutputBuffer, OutputBufferSize, &TotalReadSize, &CopyTasks](FIOBuffer CmdBuffer,
 																									   uint64	 CmdSourceOffset,
 																									   uint64	 CmdReadSize,
 																									   uint64	 OutputOffset) {
 			UNSYNC_ASSERT(OutputOffset + CmdReadSize <= OutputBufferSize);
 
 			CopyTasks.run(
-				[OutputBuffer, OutputOffset, CmdReadSize, CmdBuffer = MakeShared(std::move(CmdBuffer)), &TotalReadSize, &IoSemaphore]() {
+				[OutputBuffer, OutputOffset, CmdReadSize, CmdBuffer = MakeShared(std::move(CmdBuffer)), &TotalReadSize]() {
 					memcpy(OutputBuffer + OutputOffset, CmdBuffer->GetData(), CmdReadSize);
 					TotalReadSize += CmdReadSize;
-					IoSemaphore.Release();
+					GScheduler.FilesystemSemaphore.Release();
 				});
 		};
 
@@ -151,15 +151,13 @@ ComputeBlocksVariableT(FIOReader& Reader, const FComputeBlocksParams& Params)
 	std::vector<FTask> Tasks;
 	Tasks.resize(NumTasks);
 
-	FSemaphore IoSemaphore(MAX_ACTIVE_READERS);
-
 	FTaskGroup TaskGroup;
 
 	FBufferPool BufferPool(BytesPerTask);
 
 	for (uint64 TaskIndex = 0; TaskIndex < NumTasks; ++TaskIndex)
 	{
-		IoSemaphore.Acquire();
+		GScheduler.FilesystemSemaphore.Acquire();
 
 		const uint64 ThisTaskOffset = BytesPerTask * TaskIndex;
 		const uint64 ThisTaskSize	= CalcChunkSize(TaskIndex, BytesPerTask, InputSize);
@@ -177,7 +175,6 @@ ComputeBlocksVariableT(FIOReader& Reader, const FComputeBlocksParams& Params)
 		}
 
 		auto ScanTask = [&Tasks,
-						 &IoSemaphore,
 						 &BufferPool,
 						 &Params,
 						 MinimumBlockSize,
@@ -290,7 +287,7 @@ ComputeBlocksVariableT(FIOReader& Reader, const FComputeBlocksParams& Params)
 			HashScan<WeakHasher>(DataBegin, ThisTaskSize, MinimumBlockSize, ScanFn);
 
 			BufferPool.Release(ScanTaskBuffer);
-			IoSemaphore.Release();
+			GScheduler.FilesystemSemaphore.Release();
 		};
 
 		TaskGroup.run(ScanTask);
@@ -467,7 +464,6 @@ ComputeBlocksFixedT(FIOReader& Reader, const FComputeBlocksParams& Params)
 	std::atomic<uint64> NumBlocksCompleted = {};
 
 	{
-		FSemaphore IoSemaphore(MAX_ACTIVE_READERS);
 		FTaskGroup TaskGroup;
 
 		for (uint64 I = 0; I < NumReads; ++I)
@@ -475,9 +471,9 @@ ComputeBlocksFixedT(FIOReader& Reader, const FComputeBlocksParams& Params)
 			uint64 ThisReadSize = CalcChunkSize(I, ReadSize, Reader.GetSize());
 			uint64 Offset		= I * ReadSize;
 
-			IoSemaphore.Acquire();
+			GScheduler.FilesystemSemaphore.Acquire();
 
-			auto ReadCallback = [&NumReadsCompleted, &TaskGroup, &NumBlocksCompleted, &Blocks, &IoSemaphore, &Params, BlockSize](
+			auto ReadCallback = [&NumReadsCompleted, &TaskGroup, &NumBlocksCompleted, &Blocks, &Params, BlockSize](
 									FIOBuffer CmdBuffer,
 									uint64	  CmdOffset,
 									uint64	  CmdReadSize,
@@ -487,7 +483,6 @@ ComputeBlocksFixedT(FIOReader& Reader, const FComputeBlocksParams& Params)
 				TaskGroup.run([&NumReadsCompleted,
 							   &NumBlocksCompleted,
 							   &Blocks,
-							   &IoSemaphore,
 							   &Params,
 							   BlockSize,
 							   CmdBuffer  = MakeShared(std::move(CmdBuffer)),
@@ -521,7 +516,7 @@ ComputeBlocksFixedT(FIOReader& Reader, const FComputeBlocksParams& Params)
 
 					++NumReadsCompleted;
 
-					IoSemaphore.Release();
+					GScheduler.FilesystemSemaphore.Release();
 				});
 			};
 
@@ -814,7 +809,6 @@ DiffBlocksParallelT(FIOReader&				  BaseDataReader,
 	const uint64	   NumTasks = DivUp(BaseDataSize, BytesPerTask);
 	Tasks.resize(NumTasks);
 
-	FSemaphore IoSemaphore(MAX_ACTIVE_READERS);
 	FTaskGroup TaskGroup;
 
 	for (uint64 I = 0; I < NumTasks; ++I)
@@ -826,10 +820,10 @@ DiffBlocksParallelT(FIOReader&				  BaseDataReader,
 		Task.Offset = TaskBegin;
 		Task.Size	= TaskEnd - TaskBegin;
 
-		IoSemaphore.Acquire();
+		GScheduler.FilesystemSemaphore.Acquire();
 
 		auto ReadCallback =
-			[&SourceStrongHashSet, &SourceWeakHashSet, &Tasks, &TaskGroup, &IoSemaphore, BaseDataSize, StrongHasher, BlockSize](
+			[&SourceStrongHashSet, &SourceWeakHashSet, &Tasks, &TaskGroup, BaseDataSize, StrongHasher, BlockSize](
 				FIOBuffer CmdBuffer,
 				uint64	  CmdOffset,
 				uint64	  CmdReadSize,
@@ -837,7 +831,6 @@ DiffBlocksParallelT(FIOReader&				  BaseDataReader,
 				TaskGroup.run([&SourceStrongHashSet,
 							   &SourceWeakHashSet,
 							   &Tasks,
-							   &IoSemaphore,
 							   CmdBuffer = std::make_shared<FIOBuffer>(std::move(CmdBuffer)),
 							   CmdReadSize,
 							   CmdUserData,
@@ -904,7 +897,7 @@ DiffBlocksParallelT(FIOReader&				  BaseDataReader,
 
 					HashScan<WeakHasher>(TaskBuffer, Task.Size, BlockSize, ScanFn);
 
-					IoSemaphore.Release();
+					GScheduler.FilesystemSemaphore.Release();
 				});
 			};
 
@@ -1135,8 +1128,6 @@ ValidateTarget(FIOReader& Reader, const FNeedList& NeedList, EStrongHashAlgorith
 		return A.Offset < B.Offset;
 	});
 
-	FSemaphore IoSemaphore(MAX_ACTIVE_READERS);
-
 	const uint64		TotalStreamBytes = Reader.GetSize();
 	std::atomic<uint64> NumInvalidBlocks = {};
 	FTaskGroup			TaskGroup;
@@ -1176,7 +1167,7 @@ ValidateTarget(FIOReader& Reader, const FNeedList& NeedList, EStrongHashAlgorith
 
 		UNSYNC_ASSERT(BatchSizeBytes <= MaxBatchSizeBytes || BatchBegin == BlockIndex);
 
-		IoSemaphore.Acquire();
+		GScheduler.FilesystemSemaphore.Acquire();
 
 		const uint64 ReadOffset = ValidationBlocks[BatchBegin].Offset;
 		UNSYNC_ASSERT(BlockIndex + 1 == ValidationBlocks.size() ||
@@ -1190,7 +1181,6 @@ ValidateTarget(FIOReader& Reader, const FNeedList& NeedList, EStrongHashAlgorith
 							 BatchSizeBytes,
 							 &NumInvalidBlocks,
 							 &TaskGroup,
-							 &IoSemaphore,
 							 &ValidationProgressLogger,
 							 &ValidationBlocks](FIOBuffer CmdBuffer, uint64 CmdSourceOffset, uint64 CmdReadSize, uint64 CmdUserData) {
 			if (CmdReadSize != BatchSizeBytes)
@@ -1207,7 +1197,6 @@ ValidateTarget(FIOReader& Reader, const FNeedList& NeedList, EStrongHashAlgorith
 						   bLogVerbose,
 						   LogIndent,
 						   &NumInvalidBlocks,
-						   &IoSemaphore,
 						   &ValidationProgressLogger,
 						   &ValidationBlocks]() {
 				FLogIndentScope	   IndentScope(LogIndent, true);
@@ -1229,7 +1218,7 @@ ValidateTarget(FIOReader& Reader, const FNeedList& NeedList, EStrongHashAlgorith
 					ValidationProgressLogger.Add(Block.Size);
 				}
 
-				IoSemaphore.Release();
+				GScheduler.FilesystemSemaphore.Release();
 			});
 		};
 
@@ -2092,7 +2081,7 @@ struct FFileSyncTaskBatch
 
 		Result->BlockMap.reserve(UniqueNeedBlocks.size());
 
-		ProxyPool.ParallelDownloadSemaphore.Acquire();
+		GScheduler.DownloadSempahore.Acquire();
 		std::unique_ptr<FProxy> Proxy = ProxyPool.Alloc();
 
 		if (Proxy)
@@ -2138,7 +2127,7 @@ struct FFileSyncTaskBatch
 			UNSYNC_UNUSED(DownloadResult);
 		}
 
-		ProxyPool.ParallelDownloadSemaphore.Release();
+		GScheduler.DownloadSempahore.Release();
 
 		return Result;
 	}
diff --git a/Engine/Source/Programs/Unsync/Private/UnsyncCore.h b/Engine/Source/Programs/Unsync/Private/UnsyncCore.h
--- a/Engine/Source/Programs/Unsync/Private/UnsyncCore.h
+++ b/Engine/Source/Programs/Unsync/Private/UnsyncCore.h
@@ -22,8 +22,6 @@ namespace unsync {
 
 extern bool GDryRun;
 
-static constexpr uint32 MAX_ACTIVE_READERS = 64;
-
 class FProxy;
 class FProxyPool;
 class FBlockCache;
diff --git a/Engine/Source/Programs/Unsync/Private/UnsyncJupiter.cpp b/Engine/Source/Programs/Unsync/Private/UnsyncJupiter.cpp
--- a/Engine/Source/Programs/Unsync/Private/UnsyncJupiter.cpp
+++ b/Engine/Source/Programs/Unsync/Private/UnsyncJupiter.cpp
@@ -7,6 +7,7 @@
 #include "UnsyncPool.h"
 #include "UnsyncSerialization.h"
 #include "UnsyncThread.h"
+#include "UnsyncScheduler.h"
 
 #include <fmt/format.h>
 #if __has_include(<fmt/xchar.h>)
@@ -323,8 +324,6 @@ JupiterPutManifest(const FDirectoryManifest&	 Manifest,
 		ConnectionPool.Release(std::move(Connection));
 	}
 
-	FSemaphore ManifestUploadSemaphore(2);	// up to 2 concurrent connections
-
 	FTaskGroup ManifestUploadTasks;
 
 	struct ManifestPutTask
@@ -335,7 +334,7 @@ JupiterPutManifest(const FDirectoryManifest&	 Manifest,
 
 	for (uint64 ManifestChunkIndex = 0; ManifestChunkIndex < ChunkedRefManifestCb.size(); ++ManifestChunkIndex)
 	{
-		ManifestUploadSemaphore.Acquire();	// must be acquired before task is spawned
+		GScheduler.UploadSempahore.Acquire();	 // must be acquired before task is spawned
 
 		ManifestPutTask&	 Context	   = TaskContexts[ManifestChunkIndex];
 		const FMiniCbWriter& RefManifestCb = ChunkedRefManifestCb[ManifestChunkIndex];
@@ -354,8 +353,7 @@ JupiterPutManifest(const FDirectoryManifest&	 Manifest,
 						&RefManifestCb,
 						&ConnectionPool,
 						&Result,
-						&ResultMutex,
-						&ManifestUploadSemaphore]() {
+						&ResultMutex]() {
 			std::unique_ptr<FHttpConnection> Connection = ConnectionPool.Acquire();
 
 			FHash160	RefManifestHash	 = HashBlake3Bytes<FHash160>(RefManifestCb.Data(), RefManifestCb.Size());
@@ -388,7 +386,7 @@ JupiterPutManifest(const FDirectoryManifest&	 Manifest,
 			}
 
 			ConnectionPool.Release(std::move(Connection));
-			ManifestUploadSemaphore.Release();
+			GScheduler.UploadSempahore.Release();
 
 			if (Response.Success())
 			{
diff --git a/Engine/Source/Programs/Unsync/Private/UnsyncProxy.cpp b/Engine/Source/Programs/Unsync/Private/UnsyncProxy.cpp
--- a/Engine/Source/Programs/Unsync/Private/UnsyncProxy.cpp
+++ b/Engine/Source/Programs/Unsync/Private/UnsyncProxy.cpp
@@ -7,6 +7,7 @@
 #include "UnsyncJupiter.h"
 #include "UnsyncPool.h"
 #include "UnsyncProgress.h"
+#include "UnsyncScheduler.h"
 
 #include <fmt/format.h>
 #include <atomic>
@@ -689,7 +690,6 @@ ProxyQuery::DownloadFile(FHttpConnection&					 InConnection,
 	}
 
 	FAtomicError Error;
-	FSemaphore	 DownloadSempahore(8);	// limit concurrent connections
 
 	FLogProgressScope DownloadProgress(FileSize, ELogProgressUnits::MB);
 
@@ -706,7 +706,7 @@ ProxyQuery::DownloadFile(FHttpConnection&					 InConnection,
 	}
 
 	auto ProcessChunk =
-		[&Error, &Result, &Url, &ConnectionPool, &DownloadSempahore, &DownloadProgress, &BearerToken](
+		[&Error, &Result, &Url, &ConnectionPool, &DownloadProgress, &BearerToken](
 			const FRange& Range)
 	{
 		FLogIndentScope	   IndentScope(DownloadProgress.ParentThreadIndent, true);
@@ -717,7 +717,7 @@ ProxyQuery::DownloadFile(FHttpConnection&					 InConnection,
 			return;
 		}
 
-		DownloadSempahore.Acquire();
+		GScheduler.DownloadSempahore.Acquire();
 
 		std::unique_ptr<FHttpConnection> Connection = ConnectionPool.Acquire();
 
@@ -747,7 +747,7 @@ ProxyQuery::DownloadFile(FHttpConnection&					 InConnection,
 
 		ConnectionPool.Release(std::move(Connection));
 
-		DownloadSempahore.Release();
+		GScheduler.DownloadSempahore.Release();
 	};
 
 	ParallelForEach(Chunks, ProcessChunk);
@@ -949,8 +949,7 @@ FProxyPool::FProxyPool() : FProxyPool(FRemoteDesc(), nullptr)
 }
 
 FProxyPool::FProxyPool(const FRemoteDesc& InRemoteDesc, const FAuthDesc* InAuthDesc)
-: ParallelDownloadSemaphore(InRemoteDesc.MaxConnections)
-, RemoteDesc(InRemoteDesc)
+: RemoteDesc(InRemoteDesc)
 , AuthDesc(InAuthDesc)
 , bValid(InRemoteDesc.IsValid())
 {
diff --git a/Engine/Source/Programs/Unsync/Private/UnsyncProxy.h b/Engine/Source/Programs/Unsync/Private/UnsyncProxy.h
--- a/Engine/Source/Programs/Unsync/Private/UnsyncProxy.h
+++ b/Engine/Source/Programs/Unsync/Private/UnsyncProxy.h
@@ -166,8 +166,6 @@ public:
 	void Invalidate();
 	bool IsValid() const;
 
-	FSemaphore ParallelDownloadSemaphore;
-
 	const FRemoteDesc RemoteDesc;
 	const FAuthDesc* AuthDesc = nullptr; // optional reference to externally-owned auth parameters
 
diff --git a/Engine/Source/Programs/Unsync/Private/UnsyncScheduler.cpp b/Engine/Source/Programs/Unsync/Private/UnsyncScheduler.cpp
new file mode 100644
--- /dev/null
+++ b/Engine/Source/Programs/Unsync/Private/UnsyncScheduler.cpp
@@ -0,0 +1,20 @@
+// Copyright Epic Games, Inc. All Rights Reserved.
+
+#include "UnsyncScheduler.h"
+
+namespace unsync {
+
+FScheduler GScheduler;
+
+FScheduler::FScheduler()
+: DownloadSempahore(MAX_DOWNLOAD_TASKS)
+, UploadSempahore(MAX_UPLOAD_TASKS)
+, FilesystemSemaphore(MAX_FILESYSTEM_TASKS)
+{
+}
+
+FScheduler::~FScheduler()
+{
+}
+
+}  // namespace unsync
diff --git a/Engine/Source/Programs/Unsync/Private/UnsyncScheduler.h b/Engine/Source/Programs/Unsync/Private/UnsyncScheduler.h
new file mode 100644
--- /dev/null
+++ b/Engine/Source/Programs/Unsync/Private/UnsyncScheduler.h
@@ -0,0 +1,31 @@
+// Copyright Epic Games, Inc. All Rights Reserved.
+
+#pragma once
+
+#include "UnsyncCommon.h"
+#include "UnsyncThread.h"
+
+namespace unsync {
+
+class FScheduler
+{
+public:
+	UNSYNC_DISALLOW_COPY_ASSIGN(FScheduler)
+
+	static constexpr uint32 MAX_DOWNLOAD_TASKS	 = 8;
+	static constexpr uint32 MAX_UPLOAD_TASKS	 = 2;
+	static constexpr uint32 MAX_FILESYSTEM_TASKS = 64;
+
+	FScheduler();
+	~FScheduler();
+
+	FSemaphore DownloadSempahore;
+	FSemaphore UploadSempahore;
+	FSemaphore FilesystemSemaphore;
+
+private:
+};
+
+extern FScheduler GScheduler;
+
+}  // namespace unsync
diff --git a/Engine/Source/Programs/Unsync/Private/UnsyncTarget.cpp b/Engine/Source/Programs/Unsync/Private/UnsyncTarget.cpp
--- a/Engine/Source/Programs/Unsync/Private/UnsyncTarget.cpp
+++ b/Engine/Source/Programs/Unsync/Private/UnsyncTarget.cpp
@@ -9,6 +9,7 @@
 #include "UnsyncProxy.h"
 #include "UnsyncScavenger.h"
 #include "UnsyncThread.h"
+#include "UnsyncScheduler.h"
 
 namespace unsync {
 
@@ -85,7 +86,7 @@ DownloadBlocks(FProxyPool&					  ProxyPool,
 
 		for (FDownloadBatch Batch : Batches)
 		{
-			ProxyPool.ParallelDownloadSemaphore.Acquire();
+			GScheduler.DownloadSempahore.Acquire();
 			DownloadTasks.run(
 				[NeedBlocks,
 				 Batch,
@@ -100,7 +101,7 @@ DownloadBlocks(FProxyPool&					  ProxyPool,
 				{
 					if (bGotError)
 					{
-						ProxyPool.ParallelDownloadSemaphore.Release();
+						GScheduler.DownloadSempahore.Release();
 						return;
 					}
 
@@ -136,7 +137,7 @@ DownloadBlocks(FProxyPool&					  ProxyPool,
 							ProxyPool.Invalidate();
 						}
 					}
-					ProxyPool.ParallelDownloadSemaphore.Release();
+					GScheduler.DownloadSempahore.Release();
 				});
 		}
 
@@ -195,7 +196,6 @@ BuildTarget(FIOWriter& Output, FIOReader& Source, FIOReader& Base, const FNeedLi
 	};
 	FStats Stats;
 
-	FSemaphore WriteSemaphore(MAX_ACTIVE_READERS);	// throttle writing tasks to avoid memory bloat
 	FTaskGroup WriteTasks;
 
 	// Remember if parent thread has verbose logging and indentation
@@ -233,7 +233,7 @@ BuildTarget(FIOWriter& Output, FIOReader& Source, FIOReader& Base, const FNeedLi
 	}
 
 	auto ProcessNeedList =
-		[bAllowVerboseLog, LogIndent, &Output, &Error, &WriteSemaphore, &WriteTasks, &bWaitingForBaseData, &Stats, &Params, SizeInfo](
+		[bAllowVerboseLog, LogIndent, &Output, &Error, &WriteTasks, &bWaitingForBaseData, &Stats, &Params, SizeInfo](
 			FIOReader&					   DataProvider,
 			const std::vector<FNeedBlock>& NeedBlocks,
 			uint64						   TotalCopySize,
@@ -326,14 +326,14 @@ BuildTarget(FIOWriter& Output, FIOReader& Source, FIOReader& Base, const FNeedLi
 
 			uint64 ReadBytes = 0;
 
-			auto ReadCallback = [&ReadBytes, &Output, &WriteTasks, &Error, &WriteSemaphore, &Stats, Block, ListType](FIOBuffer CmdBuffer,
+			auto ReadCallback = [&ReadBytes, &Output, &WriteTasks, &Error, &Stats, Block, ListType](FIOBuffer CmdBuffer,
 																													 uint64	   CmdOffset,
 																													 uint64	   CmdReadSize,
 																													 uint64	   CmdUserData)
 			{
-				WriteSemaphore.Acquire();
+				GScheduler.FilesystemSemaphore.Acquire();
 				WriteTasks.run(
-					[Buffer = MakeShared(std::move(CmdBuffer)), CmdReadSize, Block, &Output, &Error, &WriteSemaphore, &Stats, ListType]()
+					[Buffer = MakeShared(std::move(CmdBuffer)), CmdReadSize, Block, &Output, &Error, &Stats, ListType]()
 					{
 						const uint64 WrittenBytes = Output.Write(Buffer->GetData(), Block.TargetOffset, CmdReadSize);
 
@@ -350,7 +350,7 @@ BuildTarget(FIOWriter& Output, FIOReader& Source, FIOReader& Base, const FNeedLi
 							UNSYNC_FATAL(L"Unexpected block list type");
 						}
 
-						WriteSemaphore.Release();
+						GScheduler.FilesystemSemaphore.Release();
 
 						if (WrittenBytes != CmdReadSize)
 						{
diff --git a/Engine/Source/Programs/Unsync/Private/UnsyncThread.cpp b/Engine/Source/Programs/Unsync/Private/UnsyncThread.cpp
--- a/Engine/Source/Programs/Unsync/Private/UnsyncThread.cpp
+++ b/Engine/Source/Programs/Unsync/Private/UnsyncThread.cpp
@@ -9,8 +9,6 @@ namespace unsync {
 
 uint32 GMaxThreads = std::min<uint32>(UNSYNC_MAX_TOTAL_THREADS, std::thread::hardware_concurrency());
 
-FThreadPool GThreadPool;
-
 void
 FThreadPool::StartWorkers(uint32 NumWorkers)
 {
diff --git a/Engine/Source/Programs/Unsync/Private/UnsyncThread.h b/Engine/Source/Programs/Unsync/Private/UnsyncThread.h
--- a/Engine/Source/Programs/Unsync/Private/UnsyncThread.h
+++ b/Engine/Source/Programs/Unsync/Private/UnsyncThread.h
@@ -135,8 +135,6 @@ private:
 	std::atomic<bool>		bShutdown;
 };
 
-extern FThreadPool GThreadPool;
-
 #if UNSYNC_USE_CONCRT
 
 // Cooperative semaphore implementation.
