unsync - Add a basic custom thread scheduler for Mac/Linux

diff --git a/Engine/Source/Programs/Unsync/CMakeLists.txt b/Engine/Source/Programs/Unsync/CMakeLists.txt
--- a/Engine/Source/Programs/Unsync/CMakeLists.txt
+++ b/Engine/Source/Programs/Unsync/CMakeLists.txt
@@ -159,7 +159,7 @@ if(${CMAKE_SYSTEM_NAME} MATCHES "Windows")
 	target_compile_definitions(unsync PRIVATE
 		UNSYNC_PLATFORM_WINDOWS=1
 		UNSYNC_PLATFORM_UNIX=0
-		UNSYNC_USE_CONCRT=1
+		UNSYNC_USE_CONCRT=0
 		WIN32_LEAN_AND_MEAN=1
 		NOMINMAX=1
 	)
diff --git a/Engine/Source/Programs/Unsync/Private/UnsyncCmdPack.cpp b/Engine/Source/Programs/Unsync/Private/UnsyncCmdPack.cpp
--- a/Engine/Source/Programs/Unsync/Private/UnsyncCmdPack.cpp
+++ b/Engine/Source/Programs/Unsync/Private/UnsyncCmdPack.cpp
@@ -8,6 +8,7 @@
 #include "UnsyncThread.h"
 #include "UnsyncError.h"
 #include "UnsyncProgress.h"
+#include "UnsyncScheduler.h"
 
 #include <stdio.h>
 #include <stdlib.h>
diff --git a/Engine/Source/Programs/Unsync/Private/UnsyncCmdQuery.cpp b/Engine/Source/Programs/Unsync/Private/UnsyncCmdQuery.cpp
--- a/Engine/Source/Programs/Unsync/Private/UnsyncCmdQuery.cpp
+++ b/Engine/Source/Programs/Unsync/Private/UnsyncCmdQuery.cpp
@@ -8,6 +8,7 @@
 #include "UnsyncUtil.h"
 #include "UnsyncAuth.h"
 #include "UnsyncPool.h"
+#include "UnsyncScheduler.h"
 
 #include <float.h>
 #include <algorithm>
@@ -306,13 +307,12 @@ CmdQuerySearch(const FCmdQueryOptions& Options)
 	struct FTaskContext
 	{
 		std::mutex				  Mutex;
-		FSemaphore				  ConnectionSemaphore = FSemaphore(8);
 		std::vector<FResultEntry> FoundEntries;
 		bool					  bParentThreadVerbose = false;
 		int32					  ParentThreadIndent   = 0;
 	};
 
-	FTaskGroup Tasks;
+	FTaskGroup	 Tasks = GScheduler->CreateTaskGroup();
 	FTaskContext Context;
 
 	Context.bParentThreadVerbose = GLogVerbose;
@@ -333,13 +333,13 @@ CmdQuerySearch(const FCmdQueryOptions& Options)
 		Request.Method		= EHttpMethod::GET;
 		Request.BearerToken = BearerToken;
 
-		Context.ConnectionSemaphore.Acquire();
+		GScheduler->NetworkSempahore.Acquire();
 
 		std::unique_ptr<FHttpConnection> Connection = ConnectionPool.Acquire();
 		FHttpResponse					 Response	= HttpRequest(*Connection, Request);
 		ConnectionPool.Release(std::move(Connection));
 
-		Context.ConnectionSemaphore.Release();
+		GScheduler->NetworkSempahore.Release();
 
 		if (!Response.Success())
 		{
diff --git a/Engine/Source/Programs/Unsync/Private/UnsyncCore.cpp b/Engine/Source/Programs/Unsync/Private/UnsyncCore.cpp
--- a/Engine/Source/Programs/Unsync/Private/UnsyncCore.cpp
+++ b/Engine/Source/Programs/Unsync/Private/UnsyncCore.cpp
@@ -30,7 +30,7 @@ UNSYNC_THIRD_PARTY_INCLUDES_START
 #include <md5-sse2.h>
 UNSYNC_THIRD_PARTY_INCLUDES_END
 
-#define UNSYNC_VERSION_STR "1.0.68"
+#define UNSYNC_VERSION_STR "1.0.70"
 
 namespace unsync {
 
@@ -62,7 +62,8 @@ BlockingReadLarge(FIOReader& Reader, uint64 Offset, uint64 Size, uint8* OutputBu
 		return TotalReadSize;
 	}
 
-	FTaskGroup CopyTasks;
+	FTaskGroup CopyTasks = GScheduler->CreateTaskGroup();
+	FSchedulerSemaphore IoSemaphore(*GScheduler, 16);
 
 	uint64 NumReads = DivUp(ClampedSize, BytesPerRead);
 	for (uint64 ReadIndex = 0; ReadIndex < NumReads; ++ReadIndex)
@@ -71,19 +72,19 @@ BlockingReadLarge(FIOReader& Reader, uint64 Offset, uint64 Size, uint8* OutputBu
 		const uint64 OutputOffset	= BytesPerRead * ReadIndex;
 		const uint64 ThisReadOffset = Offset + OutputOffset;
 
-		GScheduler.FilesystemSemaphore.Acquire();
-
-		auto ReadCallback = [OutputBuffer, OutputBufferSize, &TotalReadSize, &CopyTasks](FIOBuffer CmdBuffer,
+		auto ReadCallback = [OutputBuffer, OutputBufferSize, &TotalReadSize, &CopyTasks, &IoSemaphore](FIOBuffer CmdBuffer,
 																									   uint64	 CmdSourceOffset,
 																									   uint64	 CmdReadSize,
 																									   uint64	 OutputOffset) {
 			UNSYNC_ASSERT(OutputOffset + CmdReadSize <= OutputBufferSize);
 
+			IoSemaphore.Acquire();
 			CopyTasks.run(
-				[OutputBuffer, OutputOffset, CmdReadSize, CmdBuffer = MakeShared(std::move(CmdBuffer)), &TotalReadSize]() {
+				[OutputBuffer, OutputOffset, CmdReadSize, CmdBuffer = MakeShared(std::move(CmdBuffer)), &TotalReadSize, &IoSemaphore]()
+				{
 					memcpy(OutputBuffer + OutputOffset, CmdBuffer->GetData(), CmdReadSize);
 					TotalReadSize += CmdReadSize;
-					GScheduler.FilesystemSemaphore.Release();
+					IoSemaphore.Release();
 				});
 		};
 
@@ -151,13 +152,13 @@ ComputeBlocksVariableT(FIOReader& Reader, const FComputeBlocksParams& Params)
 	std::vector<FTask> Tasks;
 	Tasks.resize(NumTasks);
 
-	FTaskGroup TaskGroup;
+	FTaskGroup TaskGroup = GScheduler->CreateTaskGroup();
 
 	FBufferPool BufferPool(BytesPerTask);
 
 	for (uint64 TaskIndex = 0; TaskIndex < NumTasks; ++TaskIndex)
 	{
-		GScheduler.FilesystemSemaphore.Acquire();
+		GScheduler->FilesystemSemaphore.Acquire();
 
 		const uint64 ThisTaskOffset = BytesPerTask * TaskIndex;
 		const uint64 ThisTaskSize	= CalcChunkSize(TaskIndex, BytesPerTask, InputSize);
@@ -287,7 +288,7 @@ ComputeBlocksVariableT(FIOReader& Reader, const FComputeBlocksParams& Params)
 			HashScan<WeakHasher>(DataBegin, ThisTaskSize, MinimumBlockSize, ScanFn);
 
 			BufferPool.Release(ScanTaskBuffer);
-			GScheduler.FilesystemSemaphore.Release();
+			GScheduler->FilesystemSemaphore.Release();
 		};
 
 		TaskGroup.run(ScanTask);
@@ -464,14 +465,14 @@ ComputeBlocksFixedT(FIOReader& Reader, const FComputeBlocksParams& Params)
 	std::atomic<uint64> NumBlocksCompleted = {};
 
 	{
-		FTaskGroup TaskGroup;
+		FTaskGroup TaskGroup = GScheduler->CreateTaskGroup();
 
 		for (uint64 I = 0; I < NumReads; ++I)
 		{
 			uint64 ThisReadSize = CalcChunkSize(I, ReadSize, Reader.GetSize());
 			uint64 Offset		= I * ReadSize;
 
-			GScheduler.FilesystemSemaphore.Acquire();
+			GScheduler->FilesystemSemaphore.Acquire();
 
 			auto ReadCallback = [&NumReadsCompleted, &TaskGroup, &NumBlocksCompleted, &Blocks, &Params, BlockSize](
 									FIOBuffer CmdBuffer,
@@ -516,7 +517,7 @@ ComputeBlocksFixedT(FIOReader& Reader, const FComputeBlocksParams& Params)
 
 					++NumReadsCompleted;
 
-					GScheduler.FilesystemSemaphore.Release();
+					GScheduler->FilesystemSemaphore.Release();
 				});
 			};
 
@@ -809,7 +810,7 @@ DiffBlocksParallelT(FIOReader&				  BaseDataReader,
 	const uint64	   NumTasks = DivUp(BaseDataSize, BytesPerTask);
 	Tasks.resize(NumTasks);
 
-	FTaskGroup TaskGroup;
+	FTaskGroup TaskGroup = GScheduler->CreateTaskGroup();
 
 	for (uint64 I = 0; I < NumTasks; ++I)
 	{
@@ -820,7 +821,7 @@ DiffBlocksParallelT(FIOReader&				  BaseDataReader,
 		Task.Offset = TaskBegin;
 		Task.Size	= TaskEnd - TaskBegin;
 
-		GScheduler.FilesystemSemaphore.Acquire();
+		GScheduler->FilesystemSemaphore.Acquire();
 
 		auto ReadCallback =
 			[&SourceStrongHashSet, &SourceWeakHashSet, &Tasks, &TaskGroup, BaseDataSize, StrongHasher, BlockSize](
@@ -897,7 +898,7 @@ DiffBlocksParallelT(FIOReader&				  BaseDataReader,
 
 					HashScan<WeakHasher>(TaskBuffer, Task.Size, BlockSize, ScanFn);
 
-					GScheduler.FilesystemSemaphore.Release();
+					GScheduler->FilesystemSemaphore.Release();
 				});
 			};
 
@@ -1130,7 +1131,7 @@ ValidateTarget(FIOReader& Reader, const FNeedList& NeedList, EStrongHashAlgorith
 
 	const uint64		TotalStreamBytes = Reader.GetSize();
 	std::atomic<uint64> NumInvalidBlocks = {};
-	FTaskGroup			TaskGroup;
+	FTaskGroup			TaskGroup		 = GScheduler->CreateTaskGroup();
 
 	FLogProgressScope ValidationProgressLogger(TotalStreamBytes, ELogProgressUnits::MB);
 
@@ -1167,7 +1168,7 @@ ValidateTarget(FIOReader& Reader, const FNeedList& NeedList, EStrongHashAlgorith
 
 		UNSYNC_ASSERT(BatchSizeBytes <= MaxBatchSizeBytes || BatchBegin == BlockIndex);
 
-		GScheduler.FilesystemSemaphore.Acquire();
+		GScheduler->FilesystemSemaphore.Acquire();
 
 		const uint64 ReadOffset = ValidationBlocks[BatchBegin].Offset;
 		UNSYNC_ASSERT(BlockIndex + 1 == ValidationBlocks.size() ||
@@ -1218,7 +1219,7 @@ ValidateTarget(FIOReader& Reader, const FNeedList& NeedList, EStrongHashAlgorith
 					ValidationProgressLogger.Add(Block.Size);
 				}
 
-				GScheduler.FilesystemSemaphore.Release();
+				GScheduler->FilesystemSemaphore.Release();
 			});
 		};
 
@@ -2081,13 +2082,14 @@ struct FFileSyncTaskBatch
 
 		Result->BlockMap.reserve(UniqueNeedBlocks.size());
 
-		GScheduler.DownloadSempahore.Acquire();
+		GScheduler->NetworkSempahore.Acquire();
 		std::unique_ptr<FProxy> Proxy = ProxyPool.Alloc();
 
 		if (Proxy)
 		{
-			auto DownloadCallback = [&OutputCursor, &Result, &UniqueBlockSet, StrongHasher](const FDownloadedBlock& Block,
-																							FHash128				BlockHash) {
+			auto DownloadCallback =
+				[&OutputCursor, &Result, &UniqueBlockSet, StrongHasher](const FDownloadedBlock& Block, FHash128 BlockHash)
+			{
 				if (OutputCursor + Block.DecompressedSize <= Result->BlockData.Size())
 				{
 					if (UniqueBlockSet.find(BlockHash) != UniqueBlockSet.end())
@@ -2127,7 +2129,7 @@ struct FFileSyncTaskBatch
 			UNSYNC_UNUSED(DownloadResult);
 		}
 
-		GScheduler.DownloadSempahore.Release();
+		GScheduler->NetworkSempahore.Release();
 
 		return Result;
 	}
@@ -2888,8 +2890,8 @@ SyncDirectory(const FSyncDirectoryOptions& SyncOptions)
 		std::atomic<uint64> NumBackgroundTasks = {};
 		std::atomic<uint64> NumForegroundTasks = {};
 
-		FTaskGroup BackgroundTaskGroup;
-		FTaskGroup ForegroundTaskGroup;
+		FTaskGroup BackgroundTaskGroup = GScheduler->CreateTaskGroup(EWorkloadType::Background);
+		FTaskGroup ForegroundTaskGroup = GScheduler->CreateTaskGroup(EWorkloadType::Foreground);
 
 		std::atomic<uint64> BackgroundTaskMemory	   = {};
 		std::atomic<uint64> RemainingSourceBytes	   = EstimatedNeedBytesFromSource;
diff --git a/Engine/Source/Programs/Unsync/Private/UnsyncJupiter.cpp b/Engine/Source/Programs/Unsync/Private/UnsyncJupiter.cpp
--- a/Engine/Source/Programs/Unsync/Private/UnsyncJupiter.cpp
+++ b/Engine/Source/Programs/Unsync/Private/UnsyncJupiter.cpp
@@ -324,7 +324,7 @@ JupiterPutManifest(const FDirectoryManifest&	 Manifest,
 		ConnectionPool.Release(std::move(Connection));
 	}
 
-	FTaskGroup ManifestUploadTasks;
+	FTaskGroup ManifestUploadTasks = GScheduler->CreateTaskGroup();
 
 	struct ManifestPutTask
 	{
@@ -334,7 +334,7 @@ JupiterPutManifest(const FDirectoryManifest&	 Manifest,
 
 	for (uint64 ManifestChunkIndex = 0; ManifestChunkIndex < ChunkedRefManifestCb.size(); ++ManifestChunkIndex)
 	{
-		GScheduler.UploadSempahore.Acquire();	 // must be acquired before task is spawned
+		GScheduler->NetworkSempahore.Acquire();	 // must be acquired before task is spawned
 
 		ManifestPutTask&	 Context	   = TaskContexts[ManifestChunkIndex];
 		const FMiniCbWriter& RefManifestCb = ChunkedRefManifestCb[ManifestChunkIndex];
@@ -346,14 +346,8 @@ JupiterPutManifest(const FDirectoryManifest&	 Manifest,
 
 		UNSYNC_LOG_INDENT;
 
-		auto PutTask = [JupiterNamespace,
-						JupiterBucket,
-						HttpHeaders,
-						&Context,
-						&RefManifestCb,
-						&ConnectionPool,
-						&Result,
-						&ResultMutex]() {
+		auto PutTask = [JupiterNamespace, JupiterBucket, HttpHeaders, &Context, &RefManifestCb, &ConnectionPool, &Result, &ResultMutex]()
+		{
 			std::unique_ptr<FHttpConnection> Connection = ConnectionPool.Acquire();
 
 			FHash160	RefManifestHash	 = HashBlake3Bytes<FHash160>(RefManifestCb.Data(), RefManifestCb.Size());
@@ -386,7 +380,7 @@ JupiterPutManifest(const FDirectoryManifest&	 Manifest,
 			}
 
 			ConnectionPool.Release(std::move(Connection));
-			GScheduler.UploadSempahore.Release();
+			GScheduler->NetworkSempahore.Release();
 
 			if (Response.Success())
 			{
@@ -464,11 +458,9 @@ JupiterPutManifest(const FDirectoryManifest&	 Manifest,
 TResult<uint64>
 JupiterPush(const FDirectoryManifest& Manifest, const FRemoteDesc& RemoteDesc, const FTlsClientSettings& TlsSettings)
 {
-	auto CreateConnection = [RemoteDesc, TlsSettings] {
-		return new FHttpConnection(RemoteDesc.Host.Address, RemoteDesc.Host.Port, RemoteDesc.TlsRequirement, TlsSettings);
-	};
+	auto CreateConnection = [RemoteDesc, TlsSettings]
+	{ return new FHttpConnection(RemoteDesc.Host.Address, RemoteDesc.Host.Port, RemoteDesc.TlsRequirement, TlsSettings); };
 
-	FSemaphore					 ChunkUploadSemaphore(8);  // up to 8 concurrent connections
 	TObjectPool<FHttpConnection> ConnectionPool(CreateConnection);
 
 	{
@@ -489,7 +481,8 @@ JupiterPush(const FDirectoryManifest& Manifest, const FRemoteDesc& RemoteDesc, c
 		ConnectionPool.Release(std::move(Connection));
 	}
 
-	auto JupiterManifestUploadResult = [&RemoteDesc, &Manifest, &ConnectionPool] {
+	auto JupiterManifestUploadResult = [&RemoteDesc, &Manifest, &ConnectionPool]
+	{
 		UNSYNC_VERBOSE(L"Sending block references to Jupiter");
 		UNSYNC_LOG_INDENT;
 		const bool bIncludeManifestFile = true;
@@ -524,8 +517,8 @@ JupiterPush(const FDirectoryManifest& Manifest, const FRemoteDesc& RemoteDesc, c
 
 	std::atomic<int32>	NumUploadedBlocks = {};
 	std::atomic<uint64> ProcessedBytes	  = {};
-	FTaskGroup			UploadTasks;
-	std::atomic<bool>	bGotError = false;
+	FTaskGroup			UploadTasks		  = GScheduler->CreateTaskGroup(EWorkloadType::Upload);
+	std::atomic<bool>	bGotError		  = false;
 
 	for (const auto& It : Manifest.Files)
 	{
@@ -556,52 +549,52 @@ JupiterPush(const FDirectoryManifest& Manifest, const FRemoteDesc& RemoteDesc, c
 			}
 
 			auto ReadCallback =
-				[Block, &UploadTasks, &ProcessedBytes, &ConnectionPool, &ChunkUploadSemaphore, &bGotError, &RemoteDesc, &NumUploadedBlocks](
-					FIOBuffer ReadBuffer,
-					uint64	  ReadOffset,
-					uint64	  ReadReadSize,
-					uint64	  ReadUserData) {
-					// TODO: compress blocks usign Oodle
-					// Buffer compressed_buffer = compress(read_buffer.GetData(), read_buffer.GetSize());
-					// jupiter_put_compressed_blob(connection, compressed_blob_base_url, compressed_buffer, block.HashStrong);
-
-					ChunkUploadSemaphore.Acquire();	 // must be acquired before task is spawned
-					auto Task = [Block,
-								 &ProcessedBytes,
-								 &ConnectionPool,
-								 &ChunkUploadSemaphore,
-								 &bGotError,
-								 &NumUploadedBlocks,
-								 &RemoteDesc,
-								 ReadBuffer = MakeShared(std::move(ReadBuffer))] {
-						std::unique_ptr<FHttpConnection> Connection = ConnectionPool.Acquire();
-
-						bool		 bPutSucceeded = false;
-						const uint32 MaxAttempts   = 5;
-						for (uint32 AttemptIndex = 0; AttemptIndex <= MaxAttempts && !bGotError; ++AttemptIndex)
+				[Block, &UploadTasks, &ProcessedBytes, &ConnectionPool, &bGotError, &RemoteDesc, &NumUploadedBlocks](FIOBuffer ReadBuffer,
+																													 uint64	   ReadOffset,
+																													 uint64	   ReadReadSize,
+																													 uint64	   ReadUserData)
+			{
+				// TODO: compress blocks usign Oodle
+				// Buffer compressed_buffer = compress(read_buffer.GetData(), read_buffer.GetSize());
+				// jupiter_put_compressed_blob(connection, compressed_blob_base_url, compressed_buffer, block.HashStrong);
+
+				GScheduler->NetworkSempahore.Acquire();	 // must be acquired before task is spawned
+				auto Task = [Block,
+							 &ProcessedBytes,
+							 &ConnectionPool,
+							 &bGotError,
+							 &NumUploadedBlocks,
+							 &RemoteDesc,
+							 ReadBuffer = MakeShared(std::move(ReadBuffer))]
+				{
+					std::unique_ptr<FHttpConnection> Connection = ConnectionPool.Acquire();
+
+					bool		 bPutSucceeded = false;
+					const uint32 MaxAttempts   = 5;
+					for (uint32 AttemptIndex = 0; AttemptIndex <= MaxAttempts && !bGotError; ++AttemptIndex)
+					{
+						if (AttemptIndex != 0)
 						{
-							if (AttemptIndex != 0)
-							{
-								UNSYNC_WARNING(L"Retry attempt %d of %d", AttemptIndex, MaxAttempts);
-								SchedulerSleep(1000);
-							}
-
-							// TODO: pipeline PUT commands
-							UNSYNC_ASSERT(Block.HashStrong.Type == EHashType::Blake3_160);
-							if (jupiter_put_raw_blob(*Connection,
-													 RemoteDesc.StorageNamespace,
-													 RemoteDesc.HttpHeaders,
-													 ReadBuffer->GetBufferView(),
-													 Block.HashStrong.ToHash160())
-									.IsOk())
-							{
-								bPutSucceeded = true;
-								break;
-							}
+							UNSYNC_WARNING(L"Retry attempt %d of %d", AttemptIndex, MaxAttempts);
+							SchedulerSleep(1000);
 						}
 
-						ConnectionPool.Release(std::move(Connection));
-						ChunkUploadSemaphore.Release();
+						// TODO: pipeline PUT commands
+						UNSYNC_ASSERT(Block.HashStrong.Type == EHashType::Blake3_160);
+						if (jupiter_put_raw_blob(*Connection,
+												 RemoteDesc.StorageNamespace,
+												 RemoteDesc.HttpHeaders,
+												 ReadBuffer->GetBufferView(),
+												 Block.HashStrong.ToHash160())
+								.IsOk())
+						{
+							bPutSucceeded = true;
+							break;
+						}
+					}
+
+					ConnectionPool.Release(std::move(Connection));
+					GScheduler->NetworkSempahore.Release();
 
 						if (bPutSucceeded)
 						{
diff --git a/Engine/Source/Programs/Unsync/Private/UnsyncMain.cpp b/Engine/Source/Programs/Unsync/Private/UnsyncMain.cpp
--- a/Engine/Source/Programs/Unsync/Private/UnsyncMain.cpp
+++ b/Engine/Source/Programs/Unsync/Private/UnsyncMain.cpp
@@ -17,6 +17,7 @@
 #include "UnsyncTest.h"
 #include "UnsyncThread.h"
 #include "UnsyncUtil.h"
+#include "UnsyncScheduler.h"
 
 UNSYNC_THIRD_PARTY_INCLUDES_START
 #if UNSYNC_PLATFORM_WINDOWS
@@ -614,7 +615,16 @@ InnerMain(int Argc, char** Argv)
 	}
 
 	GMaxThreads = std::max(1u, GMaxThreads);
-	UNSYNC_VERBOSE2(L"Using threads: %d", GMaxThreads);
+	UNSYNC_VERBOSE(L"Using threads: %d", GMaxThreads);
+
+	// Don't count the main thread when starting the thread pool
+	const uint32 NumWorkerThreads = GMaxThreads - 1;
+
+	static FScheduler MainScheduler(NumWorkerThreads);
+
+	UNSYNC_ASSERT(GScheduler == nullptr);
+	GScheduler = &MainScheduler;
+
 	FConcurrencyPolicyScope ConcurrencyLimitScope(GMaxThreads);
 
 	if (Cli.got_subcommand(SubHash) || Cli.got_subcommand(SubPack))
@@ -1145,4 +1155,4 @@ main(int argc, char** argv)
 	}
 
 	return 1;
-}
+ }
diff --git a/Engine/Source/Programs/Unsync/Private/UnsyncManifest.cpp b/Engine/Source/Programs/Unsync/Private/UnsyncManifest.cpp
--- a/Engine/Source/Programs/Unsync/Private/UnsyncManifest.cpp
+++ b/Engine/Source/Programs/Unsync/Private/UnsyncManifest.cpp
@@ -7,6 +7,7 @@
 #include "UnsyncHashTable.h"
 #include "UnsyncSerialization.h"
 #include "UnsyncThread.h"
+#include "UnsyncScheduler.h"
 
 UNSYNC_THIRD_PARTY_INCLUDES_START
 #include <blake3.h>
@@ -177,10 +178,7 @@ UpdateDirectoryManifestBlocks(FDirectoryManifest& Result, const FPath& Root, con
 
 	uint32 NumProcessedFiles = 0;
 
-	FTaskGroup TaskGroup;
-
-	const uint32 MaxConcurrentFiles = 8;  // quickly diminishing returns past 8 concurrent files
-	FSemaphore	 Semaphore(MaxConcurrentFiles);
+	FTaskGroup TaskGroup = GScheduler->CreateTaskGroup();
 
 	uint64 NumSkippedBlocks = 0;
 	uint64 NumSkippedBytes	= 0;
@@ -202,10 +200,10 @@ UpdateDirectoryManifestBlocks(FDirectoryManifest& Result, const FPath& Root, con
 
 		FPath FilePath = Root / It.first;
 
-		Semaphore.Acquire();
+		GScheduler->NetworkSempahore.Acquire();
 
 		UNSYNC_VERBOSE(L"Computing blocks for '%ls' (%.2f MB)", FilePath.wstring().c_str(), SizeMb(It.second.Size));
-		auto BlockTask = [&FileManifest, &Semaphore, &Params, FilePath = std::move(FilePath)]()
+		auto BlockTask = [&FileManifest, &Params, FilePath = std::move(FilePath)]()
 		{
 			FNativeFile File(FilePath, EFileMode::ReadOnlyUnbuffered);
 			if (File.IsValid())
@@ -222,7 +220,8 @@ UpdateDirectoryManifestBlocks(FDirectoryManifest& Result, const FPath& Root, con
 							 FilePath.wstring().c_str(),
 							 FormatSystemErrorMessage(File.GetError()).c_str());
 			}
-			Semaphore.Release();
+
+			GScheduler->NetworkSempahore.Release();
 		};
 
 		if (Params.bAllowThreading)
@@ -287,9 +286,11 @@ CreateDirectoryManifest(const FPath& Root, const FComputeBlocksParams& Params)
 
 	FTimePoint TimeBegin = TimePointNow();
 
-	FTaskGroup	 TaskGroup;
+	FTaskGroup TaskGroup = GScheduler->CreateTaskGroup();
+
 	const uint32 MaxConcurrentFiles = 8;  // quickly diminishing returns past 8 concurrent files
-	FSemaphore	 Semaphore(MaxConcurrentFiles);
+
+	FSchedulerSemaphore Semaphore(*GScheduler, MaxConcurrentFiles);
 
 	std::mutex ResultMutex;
 	FPath	   UnsyncDirName = ".unsync";
diff --git a/Engine/Source/Programs/Unsync/Private/UnsyncProxy.cpp b/Engine/Source/Programs/Unsync/Private/UnsyncProxy.cpp
--- a/Engine/Source/Programs/Unsync/Private/UnsyncProxy.cpp
+++ b/Engine/Source/Programs/Unsync/Private/UnsyncProxy.cpp
@@ -717,7 +717,7 @@ ProxyQuery::DownloadFile(FHttpConnection&					 InConnection,
 			return;
 		}
 
-		GScheduler.DownloadSempahore.Acquire();
+		GScheduler->NetworkSempahore.Acquire();
 
 		std::unique_ptr<FHttpConnection> Connection = ConnectionPool.Acquire();
 
@@ -747,7 +747,7 @@ ProxyQuery::DownloadFile(FHttpConnection&					 InConnection,
 
 		ConnectionPool.Release(std::move(Connection));
 
-		GScheduler.DownloadSempahore.Release();
+		GScheduler->NetworkSempahore.Release();
 	};
 
 	ParallelForEach(Chunks, ProcessChunk);
diff --git a/Engine/Source/Programs/Unsync/Private/UnsyncScavenger.cpp b/Engine/Source/Programs/Unsync/Private/UnsyncScavenger.cpp
--- a/Engine/Source/Programs/Unsync/Private/UnsyncScavenger.cpp
+++ b/Engine/Source/Programs/Unsync/Private/UnsyncScavenger.cpp
@@ -6,6 +6,7 @@
 #include "UnsyncProgress.h"
 #include "UnsyncSerialization.h"
 #include "UnsyncThread.h"
+#include "UnsyncScheduler.h"
 
 namespace unsync {
 
diff --git a/Engine/Source/Programs/Unsync/Private/UnsyncScheduler.cpp b/Engine/Source/Programs/Unsync/Private/UnsyncScheduler.cpp
--- a/Engine/Source/Programs/Unsync/Private/UnsyncScheduler.cpp
+++ b/Engine/Source/Programs/Unsync/Private/UnsyncScheduler.cpp
@@ -4,17 +4,60 @@
 
 namespace unsync {
 
-FScheduler GScheduler;
+FScheduler* GScheduler = nullptr;
 
-FScheduler::FScheduler()
-: DownloadSempahore(MAX_DOWNLOAD_TASKS)
-, UploadSempahore(MAX_UPLOAD_TASKS)
-, FilesystemSemaphore(MAX_FILESYSTEM_TASKS)
+FScheduler::FScheduler(uint32 InNumWorkerThreads)
+: NumWorkerThreads(InNumWorkerThreads)
+, NetworkSempahore(*this, MAX_NETWORK_TASKS)
+, FilesystemSemaphore(*this, MAX_FILESYSTEM_TASKS)
 {
+#if UNSYNC_USE_CONCRT
+	UNSYNC_VERBOSE(L"Scheduler: PPL");
+#else
+	UNSYNC_VERBOSE(L"Scheduler: Custom");
+	ThreadPool.StartWorkers(NumWorkerThreads);
+#endif
 }
 
 FScheduler::~FScheduler()
 {
 }
 
+FTaskGroup FScheduler::CreateTaskGroup(EWorkloadType Type)
+{
+	// TODO: use separate thread pools for tasks of different types
+	UNSYNC_UNUSED(Type);
+
+#if UNSYNC_USE_CONCRT
+	return FTaskGroup();
+#else
+	return FTaskGroup(ThreadPool);
+#endif
+}
+
+#if !UNSYNC_USE_CONCRT
+
+FSchedulerSemaphore::FSchedulerSemaphore(FScheduler& InScheduler, uint32 MaxCount)
+: Scheduler(InScheduler)
+, Native(std::min(MaxCount, 1 + InScheduler.NumWorkerThreads)) // 1 for main thread + N workers
+{
+}
+
+void
+FSchedulerSemaphore::Acquire()
+{
+	while (!Native.try_acquire())
+	{
+		Scheduler.TryExecuteTask();
+	}
+}
+
+void
+FSchedulerSemaphore::Release()
+{
+	Native.release();
+}
+
+#endif
+
 }  // namespace unsync
diff --git a/Engine/Source/Programs/Unsync/Private/UnsyncScheduler.h b/Engine/Source/Programs/Unsync/Private/UnsyncScheduler.h
--- a/Engine/Source/Programs/Unsync/Private/UnsyncScheduler.h
+++ b/Engine/Source/Programs/Unsync/Private/UnsyncScheduler.h
@@ -7,25 +7,145 @@
 
 namespace unsync {
 
+class FScheduler;
+struct FTaskGroup;
+
+extern FScheduler* GScheduler;
+
+enum class EWorkloadType : uint8
+{
+	Foreground,
+	Background,
+	Download,
+	Upload,
+	FileIO,
+};
+
+#if UNSYNC_USE_CONCRT
+struct FSchedulerSemaphore : FNativeSemaphore
+{
+	explicit FSchedulerSemaphore(FScheduler& InScheduler, uint32 MaxCount) : FNativeSemaphore(MaxCount) { UNSYNC_UNUSED(InScheduler); }
+};
+#else
+struct FSchedulerSemaphore
+{
+	UNSYNC_DISALLOW_COPY_ASSIGN(FSchedulerSemaphore)
+
+	FSchedulerSemaphore(FScheduler& InScheduler, uint32 MaxCount);
+
+	void Acquire();
+	void Release();
+
+	FScheduler& Scheduler;
+
+	std::counting_semaphore<UNSYNC_MAX_TOTAL_THREADS> Native;
+};
+#endif
+
 class FScheduler
 {
 public:
 	UNSYNC_DISALLOW_COPY_ASSIGN(FScheduler)
 
-	static constexpr uint32 MAX_DOWNLOAD_TASKS	 = 8;
-	static constexpr uint32 MAX_UPLOAD_TASKS	 = 2;
+	static constexpr uint32 MAX_NETWORK_TASKS	 = 8;
 	static constexpr uint32 MAX_FILESYSTEM_TASKS = 64;
 
-	FScheduler();
+	FScheduler(uint32 InNumWorkerThreads);
 	~FScheduler();
 
-	FSemaphore DownloadSempahore;
-	FSemaphore UploadSempahore;
-	FSemaphore FilesystemSemaphore;
+	const uint32 NumWorkerThreads;
+
+	FSchedulerSemaphore NetworkSempahore;
+	FSchedulerSemaphore FilesystemSemaphore;
+
+	FTaskGroup CreateTaskGroup(EWorkloadType Type = EWorkloadType::Foreground);
+
+	void TryExecuteTask() { ThreadPool.TryExecuteTask(); }
+
+private:
+	FThreadPool ThreadPool;
+};
+
+#if UNSYNC_USE_CONCRT
+
+struct FTaskGroup : concurrency::task_group
+{
+	UNSYNC_DISALLOW_COPY_ASSIGN(FTaskGroup)
 
 private:
+	friend FScheduler;
+	FTaskGroup() : concurrency::task_group() {}
 };
 
-extern FScheduler GScheduler;
+template<typename IT, typename FT>
+inline void
+ParallelForEach(IT ItBegin, IT ItEnd, FT F)
+{
+	concurrency::parallel_for_each(ItBegin, ItEnd, F);
+}
+
+#else // UNSYNC_USE_CONCRT
+
+struct FTaskGroup
+{
+	UNSYNC_DISALLOW_COPY_ASSIGN(FTaskGroup)
+
+	template<typename F>
+	void run(F InFunction)
+	{
+		++NumStartedTasks;
+
+		ThreadPool.PushTask(
+			[&NumStartedTasks  = this->NumStartedTasks,
+			 &NumFinishedTasks = this->NumFinishedTasks,
+			 Function		   = std::forward<F>(InFunction)]() -> void
+			{
+				Function();
+				++NumFinishedTasks;
+			});
+	}
+
+	void wait()
+	{
+		while (NumFinishedTasks.load() != NumStartedTasks.load())
+		{
+			ThreadPool.TryExecuteTask();
+		}
+	}
+
+	FThreadPool&		ThreadPool;
+	std::atomic<uint64> NumStartedTasks;
+	std::atomic<uint64> NumFinishedTasks;
+
+	~FTaskGroup() { wait(); };
+
+private:
+	friend FScheduler;
+	FTaskGroup(FThreadPool& InThreadPool) : ThreadPool(InThreadPool) {}
+};
+
+template<typename IT, typename FT>
+inline void
+ParallelForEach(IT ItBegin, IT ItEnd, FT F)
+{
+	FTaskGroup CreateTaskGroup = GScheduler->CreateTaskGroup();
+
+	for (; ItBegin != ItEnd; ++ItBegin)
+	{
+		auto* It = &(*ItBegin);
+		CreateTaskGroup.run([&F, It]() { F(*It); });
+	}
+
+	CreateTaskGroup.wait();
+}
+
+#endif // UNSYNC_USE_CONCRT
+
+template<typename T, typename FT>
+inline void
+ParallelForEach(T& Container, FT F)
+{
+	ParallelForEach(std::begin(Container), std::end(Container), F);
+}
 
 }  // namespace unsync
diff --git a/Engine/Source/Programs/Unsync/Private/UnsyncTarget.cpp b/Engine/Source/Programs/Unsync/Private/UnsyncTarget.cpp
--- a/Engine/Source/Programs/Unsync/Private/UnsyncTarget.cpp
+++ b/Engine/Source/Programs/Unsync/Private/UnsyncTarget.cpp
@@ -81,12 +81,12 @@ DownloadBlocks(FProxyPool&					  ProxyPool,
 
 		UNSYNC_VERBOSE2(L"Download batches: %lld", Batches.size());
 
-		FTaskGroup DownloadTasks;
+		FTaskGroup DownloadTasks = GScheduler->CreateTaskGroup();
 		std::mutex DownloadedBlocksMutex;
 
 		for (FDownloadBatch Batch : Batches)
 		{
-			GScheduler.DownloadSempahore.Acquire();
+			GScheduler->NetworkSempahore.Acquire();
 			DownloadTasks.run(
 				[NeedBlocks,
 				 Batch,
@@ -101,7 +101,7 @@ DownloadBlocks(FProxyPool&					  ProxyPool,
 				{
 					if (bGotError)
 					{
-						GScheduler.DownloadSempahore.Release();
+						GScheduler->NetworkSempahore.Release();
 						return;
 					}
 
@@ -137,7 +137,7 @@ DownloadBlocks(FProxyPool&					  ProxyPool,
 							ProxyPool.Invalidate();
 						}
 					}
-					GScheduler.DownloadSempahore.Release();
+					GScheduler->NetworkSempahore.Release();
 				});
 		}
 
@@ -196,7 +196,7 @@ BuildTarget(FIOWriter& Output, FIOReader& Source, FIOReader& Base, const FNeedLi
 	};
 	FStats Stats;
 
-	FTaskGroup WriteTasks;
+	FTaskGroup WriteTasks = GScheduler->CreateTaskGroup();
 
 	// Remember if parent thread has verbose logging and indentation
 	const bool	 bAllowVerboseLog = GLogVerbose;
@@ -331,7 +331,7 @@ BuildTarget(FIOWriter& Output, FIOReader& Source, FIOReader& Base, const FNeedLi
 																													 uint64	   CmdReadSize,
 																													 uint64	   CmdUserData)
 			{
-				GScheduler.FilesystemSemaphore.Acquire();
+				GScheduler->FilesystemSemaphore.Acquire();
 				WriteTasks.run(
 					[Buffer = MakeShared(std::move(CmdBuffer)), CmdReadSize, Block, &Output, &Error, &Stats, ListType]()
 					{
@@ -350,7 +350,7 @@ BuildTarget(FIOWriter& Output, FIOReader& Source, FIOReader& Base, const FNeedLi
 							UNSYNC_FATAL(L"Unexpected block list type");
 						}
 
-						GScheduler.FilesystemSemaphore.Release();
+						GScheduler->FilesystemSemaphore.Release();
 
 						if (WrittenBytes != CmdReadSize)
 						{
@@ -394,7 +394,7 @@ BuildTarget(FIOWriter& Output, FIOReader& Source, FIOReader& Base, const FNeedLi
 		bCompletionFlag = true;
 	};
 
-	FTaskGroup BackgroundTasks;
+	FTaskGroup BackgroundTasks = GScheduler->CreateTaskGroup();
 	BackgroundTasks.run(
 		[SizeInfo, ProcessNeedList, &NeedList, &Base, &bBaseDataCopyTaskDone]()
 		{
@@ -477,13 +477,13 @@ BuildTarget(FIOWriter& Output, FIOReader& Source, FIOReader& Base, const FNeedLi
 			}
 		}
 
-		FTaskGroup DecompressTasks;
+		FTaskGroup DecompressTasks = GScheduler->CreateTaskGroup();
 
 		FLogProgressScope DownloadProgressLogger(EstimatedDownloadSize, ELogProgressUnits::MB);
 
 		// limit how many decompression tasks can be queued up to avoid memory bloat
 		const uint64 MaxConcurrentDecompressionTasks = 64;
-		FSemaphore	 DecompressionSemaphore(MaxConcurrentDecompressionTasks);
+		FSchedulerSemaphore	 DecompressionSemaphore(*GScheduler, MaxConcurrentDecompressionTasks);
 
 		const bool			bParentThreadVerbose = GLogVerbose;
 		const uint32		ParentThreadIndent	 = GLogIndent;
diff --git a/Engine/Source/Programs/Unsync/Private/UnsyncThread.cpp b/Engine/Source/Programs/Unsync/Private/UnsyncThread.cpp
--- a/Engine/Source/Programs/Unsync/Private/UnsyncThread.cpp
+++ b/Engine/Source/Programs/Unsync/Private/UnsyncThread.cpp
@@ -62,7 +62,7 @@ FThreadPool::PopTask(bool bWaitForSignal)
 void
 FThreadPool::PushTask(FTaskFunction&& Fun)
 {
-	if (Threads.empty())
+	if (Threads.empty() || NumRunningTasks.load() == NumWorkerThreads())
 	{
 		Fun();
 	}
@@ -81,7 +81,12 @@ FThreadPool::DoWorkInternal(bool bWaitForSignal)
 
 	if (Task)
 	{
+		NumRunningTasks++;
+
 		Task();
+
+		NumRunningTasks--;
+
 		return true;
 	}
 	else
diff --git a/Engine/Source/Programs/Unsync/Private/UnsyncThread.h b/Engine/Source/Programs/Unsync/Private/UnsyncThread.h
--- a/Engine/Source/Programs/Unsync/Private/UnsyncThread.h
+++ b/Engine/Source/Programs/Unsync/Private/UnsyncThread.h
@@ -44,6 +44,8 @@ static constexpr uint32 UNSYNC_MAX_TOTAL_THREADS = 64;
 
 extern uint32 GMaxThreads;
 
+class FScheduler;
+
 class FConcurrencyPolicyScope
 {
 public:
@@ -120,6 +122,8 @@ public:
 	// Returns false if queue is empty, which may happen if worker threads have picked up the tasks already.
 	bool TryExecuteTask() { return DoWorkInternal(false); }
 
+	uint32 NumWorkerThreads() const { return uint32(Threads.size()); }
+
 private:
 
 	// Try to execute a task and return whether there may be more tasks to run
@@ -133,6 +137,7 @@ private:
 	std::mutex				Mutex;
 	std::condition_variable WorkerWakeupCondition;
 	std::atomic<bool>		bShutdown;
+	std::atomic<uint64>		NumRunningTasks;
 };
 
 #if UNSYNC_USE_CONCRT
@@ -140,14 +145,14 @@ private:
 // Cooperative semaphore implementation.
 // Using this is necessary to avoid deadlocks on low-core machines.
 // https://docs.microsoft.com/en-us/cpp/parallel/concrt/how-to-use-the-context-class-to-implement-a-cooperative-semaphore?view=msvc-160
-class FSemaphore
+class FNativeSemaphore
 {
 public:
-	UNSYNC_DISALLOW_COPY_ASSIGN(FSemaphore)
+	UNSYNC_DISALLOW_COPY_ASSIGN(FNativeSemaphore)
 
-	explicit FSemaphore(uint32 MaxCount) : Counter(MaxCount) {}
+	explicit FNativeSemaphore(uint32 MaxCount) : Counter(MaxCount) {}
 
-	~FSemaphore() {}
+	~FNativeSemaphore() {}
 
 	void Acquire()
 	{
@@ -176,31 +181,22 @@ private:
 	concurrency::concurrent_queue<concurrency::Context*> WaitingQueue;
 };
 
-using FTaskGroup = concurrency::task_group;
-
-template<typename IT, typename FT>
-inline void
-ParallelForEach(IT ItBegin, IT ItEnd, FT F)
-{
-	concurrency::parallel_for_each(ItBegin, ItEnd, F);
-}
-
 #else  // UNSYNC_USE_CONCRT
 
 #if UNSYNC_USE_MACH_SEMAPHORE
 
-class FSemaphore
+class FNativeSemaphore
 {
 public:
-	UNSYNC_DISALLOW_COPY_ASSIGN(FSemaphore)
+	UNSYNC_DISALLOW_COPY_ASSIGN(FNativeSemaphore)
 
-	explicit FSemaphore(uint32 MaxCount)
+	explicit FNativeSemaphore(uint32 MaxCount)
 	{
 		kern_return_t InitResult = semaphore_create(mach_task_self(), &Native, SYNC_POLICY_FIFO, MaxCount);
 		UNSYNC_ASSERTF(InitResult == KERN_SUCCESS, L"Failed to create a semaphore, error code: %d %hs", InitResult, mach_error_string(InitResult));
 	}
 
-	~FSemaphore()
+	~FNativeSemaphore()
 	{
 		kern_return_t DestroyResult = semaphore_destroy(mach_task_self(), Native);
 		UNSYNC_ASSERTF(DestroyResult == KERN_SUCCESS, L"Failed to destroy a semaphore, error code: %d %hs", DestroyResult, mach_error_string(DestroyResult));
@@ -223,14 +219,14 @@ public:
 
 #else // UNSYNC_USE_MACH_SEMAPHORE
 
-class FSemaphore
+class FNativeSemaphore
 {
 public:
-	UNSYNC_DISALLOW_COPY_ASSIGN(FSemaphore)
+	UNSYNC_DISALLOW_COPY_ASSIGN(FNativeSemaphore)
 
-	explicit FSemaphore(uint32 MaxCount) : Native(MaxCount) {}
+	explicit FNativeSemaphore(uint32 MaxCount) : Native(MaxCount) {}
 
-	~FSemaphore() {}
+	~FNativeSemaphore() {}
 
 	void Acquire() { Native.acquire(); }
 
@@ -243,34 +239,6 @@ private:
 
 #endif // UNSYNC_USE_MACH_SEMAPHORE
 
-// Single-threaded task group implementation
-struct FTaskGroup
-{
-	template<typename F>
-	void run(F f)
-	{
-		f();
-	}
-	void wait() {}
-};
-
-template<typename IT, typename FT>
-inline void
-ParallelForEach(IT ItBegin, IT ItEnd, FT F)
-{
-	for (; ItBegin != ItEnd; ++ItBegin)
-	{
-		F(*ItBegin);
-	}
-}
-
 #endif	// UNSYNC_USE_CONCRT
 
-template<typename T, typename FT>
-inline void
-ParallelForEach(T& Container, FT F)
-{
-	ParallelForEach(std::begin(Container), std::end(Container), F);
-}
-
 }  // namespace unsync
