unsync - Added a mechanism for unified proxy server and source directory specification to simply sync command line

Example command line `unsync.exe sync https://foo.bar/directory/path <dest>` is equivalent to `unsync.exe sync --proxy https://foo.bar directory/path <dest>`

diff --git a/Engine/Source/Programs/Unsync/Private/UnsyncCmdQuery.cpp b/Engine/Source/Programs/Unsync/Private/UnsyncCmdQuery.cpp
--- a/Engine/Source/Programs/Unsync/Private/UnsyncCmdQuery.cpp
+++ b/Engine/Source/Programs/Unsync/Private/UnsyncCmdQuery.cpp
@@ -498,6 +498,12 @@ CmdQueryFile(const FCmdQueryOptions& Options)
 int32
 CmdQuery(const FCmdQueryOptions& Options)
 {
+	if (!Options.Remote.IsValid())
+	{
+		UNSYNC_ERROR(L"Server address is not specified or is invalid");
+		return 1;
+	}
+
 	if (Options.Query == "mirrors")
 	{
 		return CmdQueryMirrors(Options);
@@ -506,7 +512,7 @@ CmdQuery(const FCmdQueryOptions& Options)
 	{
 		return CmdQueryList(Options);
 	}
-	else if (Options.Query == "search")
+	else if (Options.Query == "search" || Options.Query == "explore")
 	{
 		return CmdQuerySearch(Options);
 	}
diff --git a/Engine/Source/Programs/Unsync/Private/UnsyncCore.cpp b/Engine/Source/Programs/Unsync/Private/UnsyncCore.cpp
--- a/Engine/Source/Programs/Unsync/Private/UnsyncCore.cpp
+++ b/Engine/Source/Programs/Unsync/Private/UnsyncCore.cpp
@@ -30,7 +30,7 @@ UNSYNC_THIRD_PARTY_INCLUDES_START
 #include <md5-sse2.h>
 UNSYNC_THIRD_PARTY_INCLUDES_END
 
-#define UNSYNC_VERSION_STR "1.0.70"
+#define UNSYNC_VERSION_STR "1.0.71"
 
 namespace unsync {
 
diff --git a/Engine/Source/Programs/Unsync/Private/UnsyncHash.cpp b/Engine/Source/Programs/Unsync/Private/UnsyncHash.cpp
--- a/Engine/Source/Programs/Unsync/Private/UnsyncHash.cpp
+++ b/Engine/Source/Programs/Unsync/Private/UnsyncHash.cpp
@@ -90,40 +90,6 @@ ComputeHash(const FBuffer& Buffer, EStrongHashAlgorithmID Algorithm)
 	return ComputeHash(Buffer.Data(), Buffer.Size(), Algorithm);
 }
 
-bool
-LooksLikeHash160(const std::string_view Str)
-{
-	if (Str.length() != 40)
-	{
-		return false;
-	}
-
-	const char* PossibleChars = "0123456789abcdefABCDEF";
-	if (Str.find_first_not_of(PossibleChars) != std::string::npos)
-	{
-		return false;
-	}
-
-	return true;
-}
-
-bool
-LooksLikeHash160(const std::wstring_view Str)
-{
-	if (Str.length() != 40)
-	{
-		return false;
-	}
-
-	const wchar_t* PossibleChars = L"0123456789abcdefABCDEF";
-	if (Str.find_first_not_of(PossibleChars) != std::wstring::npos)
-	{
-		return false;
-	}
-
-	return true;
-}
-
 FGenericHash
 ComputeHash(const FBufferView& Buffer, EStrongHashAlgorithmID Algorithm)
 {
diff --git a/Engine/Source/Programs/Unsync/Private/UnsyncHash.h b/Engine/Source/Programs/Unsync/Private/UnsyncHash.h
--- a/Engine/Source/Programs/Unsync/Private/UnsyncHash.h
+++ b/Engine/Source/Programs/Unsync/Private/UnsyncHash.h
@@ -307,9 +307,6 @@ FGenericHash ComputeHash(const uint8* Data, uint64 Size, EStrongHashAlgorithmID
 FGenericHash ComputeHash(const FBuffer& Buffer, EStrongHashAlgorithmID Algorithm);
 FGenericHash ComputeHash(const FBufferView& Buffer, EStrongHashAlgorithmID Algorithm);
 
-bool LooksLikeHash160(const std::string_view Str);
-bool LooksLikeHash160(const std::wstring_view Str);
-
 }  // namespace unsync
 
 namespace std {
diff --git a/Engine/Source/Programs/Unsync/Private/UnsyncMain.cpp b/Engine/Source/Programs/Unsync/Private/UnsyncMain.cpp
--- a/Engine/Source/Programs/Unsync/Private/UnsyncMain.cpp
+++ b/Engine/Source/Programs/Unsync/Private/UnsyncMain.cpp
@@ -537,21 +537,36 @@ InnerMain(int Argc, char** Argv)
 	FRemoteDesc RemoteDesc;
 	FAuthDesc	AuthDesc;
 
-	if (RemoteAddressUtf8.empty())
+	bool bFilesystemSource = true;
+
+	std::string_view PossibleUrl;
+
+	if (Cli.got_subcommand(SubSync))
 	{
-		// Derive remote server address from source name if explicit --proxy or --remote option is not provided for sync command.
+		PossibleUrl = SourceFilenameUtf8;
+	}
+	else if (Cli.got_subcommand(SubQuery) && !QueryArgsUtf8.empty())
+	{
+		PossibleUrl = QueryArgsUtf8[0];
+	}
 
-		if (Cli.got_subcommand(SubSync) && !PathExists(SourceFilenameUtf8))
+	if (RemoteAddressUtf8.empty() && LooksLikeUrl(PossibleUrl) && (Cli.got_subcommand(SubSync) || Cli.got_subcommand(SubQuery)))
+	{
+		// Derive remote server address from source name if explicit --proxy or --remote option is not provided for sync or query
+
+		TResult<FRemoteDesc> ParsedRemoteDesc = FRemoteDesc::FromUrl(PossibleUrl);
+
+		if (ParsedRemoteDesc.IsOk())
 		{
-			TResult<FRemoteDesc> ParsedRemoteDesc = FRemoteDesc::FromUrl(SourceFilenameUtf8);
-			if (ParsedRemoteDesc.IsOk())
-			{
-				RemoteDesc = *ParsedRemoteDesc;
+			RemoteDesc = *ParsedRemoteDesc;
+			bFilesystemSource = false;
 
+			if (RemoteDesc.Protocol == EProtocolFlavor::Jupiter)
+			{
 				size_t SlashPos = RemoteDesc.StorageNamespace.find_first_of('/');
 				if (SlashPos == std::string::npos)
 				{
-					UNSYNC_ERROR(L"URL source is expected to follow [transport://]address[:port]#namespace/object format");
+					UNSYNC_ERROR(L"Jupiter URL source is expected to follow [transport://]address[:port]#namespace/object format");
 					return 1;
 				}
 				else
@@ -560,14 +575,27 @@ InnerMain(int Argc, char** Argv)
 					RemoteDesc.StorageNamespace = RemoteDesc.StorageNamespace.substr(0, SlashPos);
 				}
 			}
-			else
+			else if (RemoteDesc.Protocol == EProtocolFlavor::Unsync)
 			{
-				UNSYNC_ERROR(L"Failed to parse remote address '%hs': %ls",
-							 RemoteAddressUtf8.c_str(),
-							 ParsedRemoteDesc.TryError()->Context.c_str());
-				return 1;
+				bShouldLogin = true; // Try to authenticate by default when source is a valid URL
+
+				if (Cli.got_subcommand(SubQuery))
+				{
+					QueryArgsUtf8[0] = RemoteDesc.RequestPath;
+				}
+				else if (Cli.got_subcommand(SubSync))
+				{
+					SourceFilenameUtf8 = RemoteDesc.RequestPath;
+				}
 			}
 		}
+		else
+		{
+			UNSYNC_ERROR(L"Failed to parse remote address '%hs': %ls",
+							RemoteAddressUtf8.c_str(),
+							ParsedRemoteDesc.TryError()->Context.c_str());
+			return 1;
+		}
 	}
 	else
 	{
@@ -589,12 +617,13 @@ InnerMain(int Argc, char** Argv)
 	FPath InputFilename2		 = NormalizeFilenameUtf8(InputFilename2Utf8);
 	FPath OutputFilename		 = NormalizeFilenameUtf8(OutputFilenameUtf8);
 	FPath BaseFilename			 = NormalizeFilenameUtf8(BaseFilenameUtf8);
-	FPath SourceFilename		 = NormalizeFilenameUtf8(SourceFilenameUtf8);
 	FPath TargetFilename		 = NormalizeFilenameUtf8(TargetFilenameUtf8);
 	FPath PatchFilename			 = NormalizeFilenameUtf8(PatchFilenameUtf8);
 	FPath ScavengeRoot			 = NormalizeFilenameUtf8(ScavengeRootUtf8);
 	FPath SourceManifestFilename = NormalizeFilenameUtf8(SourceManifestFilenameUtf8);
 
+	FPath SourceFilename = bFilesystemSource ? NormalizeFilenameUtf8(SourceFilenameUtf8) : FPath(SourceFilenameUtf8);
+
 	if (GLogVeryVerbose)
 	{
 		UNSYNC_LOG(L"Very verbose logging is enabled");
diff --git a/Engine/Source/Programs/Unsync/Private/UnsyncRemote.cpp b/Engine/Source/Programs/Unsync/Private/UnsyncRemote.cpp
--- a/Engine/Source/Programs/Unsync/Private/UnsyncRemote.cpp
+++ b/Engine/Source/Programs/Unsync/Private/UnsyncRemote.cpp
@@ -178,7 +178,9 @@ FRemoteDesc::FromUrl(std::string_view Url)
 	{
 		TlsRequirement = ETlsRequirement::Required;
 	}
-	else if (Result.Protocol == EProtocolFlavor::Unsync && TlsRequirement < ETlsRequirement::Required)
+	else if (Result.Protocol == EProtocolFlavor::Unsync
+		&& TlsRequirement < ETlsRequirement::Required
+		&& Transport == ETransportProtocol::Unsync)
 	{
 		TlsRequirement = ETlsRequirement::Preferred;
 	}
diff --git a/Engine/Source/Programs/Unsync/Private/UnsyncScheduler.h b/Engine/Source/Programs/Unsync/Private/UnsyncScheduler.h
--- a/Engine/Source/Programs/Unsync/Private/UnsyncScheduler.h
+++ b/Engine/Source/Programs/Unsync/Private/UnsyncScheduler.h
@@ -128,15 +128,15 @@ template<typename IT, typename FT>
 inline void
 ParallelForEach(IT ItBegin, IT ItEnd, FT F)
 {
-	FTaskGroup CreateTaskGroup = GScheduler->CreateTaskGroup();
+	FTaskGroup TaskGroup = GScheduler->CreateTaskGroup();
 
 	for (; ItBegin != ItEnd; ++ItBegin)
 	{
 		auto* It = &(*ItBegin);
-		CreateTaskGroup.run([&F, It]() { F(*It); });
+		TaskGroup.run([&F, It]() { F(*It); });
 	}
 
-	CreateTaskGroup.wait();
+	TaskGroup.wait();
 }
 
 #endif // UNSYNC_USE_CONCRT
diff --git a/Engine/Source/Programs/Unsync/Private/UnsyncUtil.cpp b/Engine/Source/Programs/Unsync/Private/UnsyncUtil.cpp
--- a/Engine/Source/Programs/Unsync/Private/UnsyncUtil.cpp
+++ b/Engine/Source/Programs/Unsync/Private/UnsyncUtil.cpp
@@ -634,4 +634,63 @@ GetAnonymizedMachineIdString(std::string_view Seed)
 	return Result;
 }
 
+bool
+LooksLikeHash160(const std::string_view Str)
+{
+	if (Str.length() != 40)
+	{
+		return false;
+	}
+
+	const char* PossibleChars = "0123456789abcdefABCDEF";
+	if (Str.find_first_not_of(PossibleChars) != std::string::npos)
+	{
+		return false;
+	}
+
+	return true;
+}
+
+bool
+LooksLikeHash160(const std::wstring_view Str)
+{
+	if (Str.length() != 40)
+	{
+		return false;
+	}
+
+	const wchar_t* PossibleChars = L"0123456789abcdefABCDEF";
+	if (Str.find_first_not_of(PossibleChars) != std::wstring::npos)
+	{
+		return false;
+	}
+
+	return true;
+}
+
+bool
+LooksLikeUrl(std::string_view Str)
+{
+	std::string_view Prefixes[] = {
+		"http://",
+		"https://",
+		"unsync://",
+		"unsync+tls://",
+		"unsync+http://",
+		"unsync+https://",
+		"jupiter+http://",
+		"jupiter+https://",
+	};
+
+	for (std::string_view Prefix : Prefixes)
+	{
+		if (Str.starts_with(Prefix))
+		{
+			return true;
+		}
+	}
+	
+	return false;
+}
+
 }  // namespace unsync
diff --git a/Engine/Source/Programs/Unsync/Private/UnsyncUtil.h b/Engine/Source/Programs/Unsync/Private/UnsyncUtil.h
--- a/Engine/Source/Programs/Unsync/Private/UnsyncUtil.h
+++ b/Engine/Source/Programs/Unsync/Private/UnsyncUtil.h
@@ -374,4 +374,8 @@ std::string GetAnonymizedMachineIdString(std::string_view Salt = {});
 // Returns string in format 'Error code 123: Some description.`
 std::string FormatSystemErrorMessage(int32 ErrorCode);
 
+bool LooksLikeUrl(std::string_view Str);
+bool LooksLikeHash160(const std::string_view Str);
+bool LooksLikeHash160(const std::wstring_view Str);
+
 }  // namespace unsync
