unsync - Redirect warning and error log to stderr by default

diff --git a/Engine/Source/Programs/Unsync/Private/UnsyncLog.cpp b/Engine/Source/Programs/Unsync/Private/UnsyncLog.cpp
--- a/Engine/Source/Programs/Unsync/Private/UnsyncLog.cpp
+++ b/Engine/Source/Programs/Unsync/Private/UnsyncLog.cpp
@@ -70,20 +70,16 @@ LogSaveCommandLineUtf8(int Argc, char** Argv)
 static FILE*
 GetLogStream(ELogLevel LogLevel)
 {
+	// If machine readable log format is enabled, *all* diagnostic logging is sent to stderr.
+	// Otherwise, only errors and warnings are sent to stderr, while normal diagnostics are sent to stdout.
+
 	if (GLogMachineReadable)
 	{
-		if (LogLevel == ELogLevel::MachineReadable)
-		{
-			return stdout;
-		}
-		else
-		{
-			return stderr;
-		}
+		return (LogLevel == ELogLevel::MachineReadable) ? stdout : stderr;
 	}
 	else
 	{
-		return stdout;
+		return (LogLevel >= ELogLevel::Info) ? stdout : stderr;
 	}
 }
 
