unsync - Remove ConcRT/PPL, use custom scheduler by default

diff --git a/Engine/Source/Programs/Unsync/CMakeLists.txt b/Engine/Source/Programs/Unsync/CMakeLists.txt
--- a/Engine/Source/Programs/Unsync/CMakeLists.txt
+++ b/Engine/Source/Programs/Unsync/CMakeLists.txt
@@ -159,7 +159,6 @@ if(${CMAKE_SYSTEM_NAME} MATCHES "Windows")
 	target_compile_definitions(unsync PRIVATE
 		UNSYNC_PLATFORM_WINDOWS=1
 		UNSYNC_PLATFORM_UNIX=0
-		UNSYNC_USE_CONCRT=0
 		WIN32_LEAN_AND_MEAN=1
 		NOMINMAX=1
 	)
@@ -167,7 +166,6 @@ else()
 	target_compile_definitions(unsync PRIVATE 
 		UNSYNC_PLATFORM_UNIX=1
 		UNSYNC_PLATFORM_WINDOWS=0
-		UNSYNC_USE_CONCRT=0
 	)
 endif()
 
diff --git a/Engine/Source/Programs/Unsync/Private/UnsyncCmdQuery.cpp b/Engine/Source/Programs/Unsync/Private/UnsyncCmdQuery.cpp
--- a/Engine/Source/Programs/Unsync/Private/UnsyncCmdQuery.cpp
+++ b/Engine/Source/Programs/Unsync/Private/UnsyncCmdQuery.cpp
@@ -333,13 +333,13 @@ CmdQuerySearch(const FCmdQueryOptions& Options)
 		Request.Method		= EHttpMethod::GET;
 		Request.BearerToken = BearerToken;
 
-		GScheduler->NetworkSempahore.Acquire();
+		GScheduler->NetworkSemaphore.Acquire(false);
 
 		std::unique_ptr<FHttpConnection> Connection = ConnectionPool.Acquire();
 		FHttpResponse					 Response	= HttpRequest(*Connection, Request);
 		ConnectionPool.Release(std::move(Connection));
 
-		GScheduler->NetworkSempahore.Release();
+		GScheduler->NetworkSemaphore.Release();
 
 		if (!Response.Success())
 		{
diff --git a/Engine/Source/Programs/Unsync/Private/UnsyncCore.cpp b/Engine/Source/Programs/Unsync/Private/UnsyncCore.cpp
--- a/Engine/Source/Programs/Unsync/Private/UnsyncCore.cpp
+++ b/Engine/Source/Programs/Unsync/Private/UnsyncCore.cpp
@@ -30,7 +30,7 @@ UNSYNC_THIRD_PARTY_INCLUDES_START
 #include <md5-sse2.h>
 UNSYNC_THIRD_PARTY_INCLUDES_END
 
-#define UNSYNC_VERSION_STR "1.0.71"
+#define UNSYNC_VERSION_STR "1.0.72"
 
 namespace unsync {
 
@@ -62,8 +62,8 @@ BlockingReadLarge(FIOReader& Reader, uint64 Offset, uint64 Size, uint8* OutputBu
 		return TotalReadSize;
 	}
 
-	FTaskGroup CopyTasks = GScheduler->CreateTaskGroup();
 	FSchedulerSemaphore IoSemaphore(*GScheduler, 16);
+	FTaskGroup			CopyTasks = GScheduler->CreateTaskGroup(&IoSemaphore);
 
 	uint64 NumReads = DivUp(ClampedSize, BytesPerRead);
 	for (uint64 ReadIndex = 0; ReadIndex < NumReads; ++ReadIndex)
@@ -72,19 +72,17 @@ BlockingReadLarge(FIOReader& Reader, uint64 Offset, uint64 Size, uint8* OutputBu
 		const uint64 OutputOffset	= BytesPerRead * ReadIndex;
 		const uint64 ThisReadOffset = Offset + OutputOffset;
 
-		auto ReadCallback = [OutputBuffer, OutputBufferSize, &TotalReadSize, &CopyTasks, &IoSemaphore](FIOBuffer CmdBuffer,
+		auto ReadCallback = [OutputBuffer, OutputBufferSize, &TotalReadSize, &CopyTasks](FIOBuffer CmdBuffer,
 																									   uint64	 CmdSourceOffset,
 																									   uint64	 CmdReadSize,
 																									   uint64	 OutputOffset) {
 			UNSYNC_ASSERT(OutputOffset + CmdReadSize <= OutputBufferSize);
 
-			IoSemaphore.Acquire();
 			CopyTasks.run(
-				[OutputBuffer, OutputOffset, CmdReadSize, CmdBuffer = MakeShared(std::move(CmdBuffer)), &TotalReadSize, &IoSemaphore]()
+				[OutputBuffer, OutputOffset, CmdReadSize, CmdBuffer = MakeShared(std::move(CmdBuffer)), &TotalReadSize]()
 				{
 					memcpy(OutputBuffer + OutputOffset, CmdBuffer->GetData(), CmdReadSize);
 					TotalReadSize += CmdReadSize;
-					IoSemaphore.Release();
 				});
 		};
 
@@ -152,13 +150,14 @@ ComputeBlocksVariableT(FIOReader& Reader, const FComputeBlocksParams& Params)
 	std::vector<FTask> Tasks;
 	Tasks.resize(NumTasks);
 
+	FSchedulerSemaphore IoSemaphore(*GScheduler, 16);
 	FTaskGroup TaskGroup = GScheduler->CreateTaskGroup();
 
 	FBufferPool BufferPool(BytesPerTask);
 
 	for (uint64 TaskIndex = 0; TaskIndex < NumTasks; ++TaskIndex)
 	{
-		GScheduler->FilesystemSemaphore.Acquire();
+		IoSemaphore.Acquire(); // TODO: remove
 
 		const uint64 ThisTaskOffset = BytesPerTask * TaskIndex;
 		const uint64 ThisTaskSize	= CalcChunkSize(TaskIndex, BytesPerTask, InputSize);
@@ -176,6 +175,7 @@ ComputeBlocksVariableT(FIOReader& Reader, const FComputeBlocksParams& Params)
 		}
 
 		auto ScanTask = [&Tasks,
+						 &IoSemaphore,
 						 &BufferPool,
 						 &Params,
 						 MinimumBlockSize,
@@ -288,7 +288,7 @@ ComputeBlocksVariableT(FIOReader& Reader, const FComputeBlocksParams& Params)
 			HashScan<WeakHasher>(DataBegin, ThisTaskSize, MinimumBlockSize, ScanFn);
 
 			BufferPool.Release(ScanTaskBuffer);
-			GScheduler->FilesystemSemaphore.Release();
+			IoSemaphore.Release();
 		};
 
 		TaskGroup.run(ScanTask);
@@ -465,15 +465,14 @@ ComputeBlocksFixedT(FIOReader& Reader, const FComputeBlocksParams& Params)
 	std::atomic<uint64> NumBlocksCompleted = {};
 
 	{
-		FTaskGroup TaskGroup = GScheduler->CreateTaskGroup();
+		FSchedulerSemaphore IoSemaphore(*GScheduler, 16);
+		FTaskGroup			TaskGroup = GScheduler->CreateTaskGroup(&IoSemaphore);
 
 		for (uint64 I = 0; I < NumReads; ++I)
 		{
 			uint64 ThisReadSize = CalcChunkSize(I, ReadSize, Reader.GetSize());
 			uint64 Offset		= I * ReadSize;
 
-			GScheduler->FilesystemSemaphore.Acquire();
-
 			auto ReadCallback = [&NumReadsCompleted, &TaskGroup, &NumBlocksCompleted, &Blocks, &Params, BlockSize](
 									FIOBuffer CmdBuffer,
 									uint64	  CmdOffset,
@@ -516,8 +515,6 @@ ComputeBlocksFixedT(FIOReader& Reader, const FComputeBlocksParams& Params)
 					}
 
 					++NumReadsCompleted;
-
-					GScheduler->FilesystemSemaphore.Release();
 				});
 			};
 
@@ -810,7 +807,8 @@ DiffBlocksParallelT(FIOReader&				  BaseDataReader,
 	const uint64	   NumTasks = DivUp(BaseDataSize, BytesPerTask);
 	Tasks.resize(NumTasks);
 
-	FTaskGroup TaskGroup = GScheduler->CreateTaskGroup();
+	FSchedulerSemaphore IoSemaphore(*GScheduler, 16);
+	FTaskGroup			TaskGroup = GScheduler->CreateTaskGroup(&IoSemaphore);
 
 	for (uint64 I = 0; I < NumTasks; ++I)
 	{
@@ -821,8 +819,6 @@ DiffBlocksParallelT(FIOReader&				  BaseDataReader,
 		Task.Offset = TaskBegin;
 		Task.Size	= TaskEnd - TaskBegin;
 
-		GScheduler->FilesystemSemaphore.Acquire();
-
 		auto ReadCallback =
 			[&SourceStrongHashSet, &SourceWeakHashSet, &Tasks, &TaskGroup, BaseDataSize, StrongHasher, BlockSize](
 				FIOBuffer CmdBuffer,
@@ -897,8 +893,6 @@ DiffBlocksParallelT(FIOReader&				  BaseDataReader,
 					};
 
 					HashScan<WeakHasher>(TaskBuffer, Task.Size, BlockSize, ScanFn);
-
-					GScheduler->FilesystemSemaphore.Release();
 				});
 			};
 
@@ -1131,7 +1125,8 @@ ValidateTarget(FIOReader& Reader, const FNeedList& NeedList, EStrongHashAlgorith
 
 	const uint64		TotalStreamBytes = Reader.GetSize();
 	std::atomic<uint64> NumInvalidBlocks = {};
-	FTaskGroup			TaskGroup		 = GScheduler->CreateTaskGroup();
+	FSchedulerSemaphore IoSemaphore(*GScheduler, 16);
+	FTaskGroup			TaskGroup		 = GScheduler->CreateTaskGroup(&IoSemaphore);
 
 	FLogProgressScope ValidationProgressLogger(TotalStreamBytes, ELogProgressUnits::MB);
 
@@ -1168,8 +1163,6 @@ ValidateTarget(FIOReader& Reader, const FNeedList& NeedList, EStrongHashAlgorith
 
 		UNSYNC_ASSERT(BatchSizeBytes <= MaxBatchSizeBytes || BatchBegin == BlockIndex);
 
-		GScheduler->FilesystemSemaphore.Acquire();
-
 		const uint64 ReadOffset = ValidationBlocks[BatchBegin].Offset;
 		UNSYNC_ASSERT(BlockIndex + 1 == ValidationBlocks.size() ||
 					  (ReadOffset + BatchSizeBytes) == ValidationBlocks[BlockIndex + 1].Offset);
@@ -1218,8 +1211,6 @@ ValidateTarget(FIOReader& Reader, const FNeedList& NeedList, EStrongHashAlgorith
 
 					ValidationProgressLogger.Add(Block.Size);
 				}
-
-				GScheduler->FilesystemSemaphore.Release();
 			});
 		};
 
@@ -2082,7 +2073,7 @@ struct FFileSyncTaskBatch
 
 		Result->BlockMap.reserve(UniqueNeedBlocks.size());
 
-		GScheduler->NetworkSempahore.Acquire();
+		GScheduler->NetworkSemaphore.Acquire(false);
 		std::unique_ptr<FProxy> Proxy = ProxyPool.Alloc();
 
 		if (Proxy)
@@ -2129,7 +2120,7 @@ struct FFileSyncTaskBatch
 			UNSYNC_UNUSED(DownloadResult);
 		}
 
-		GScheduler->NetworkSempahore.Release();
+		GScheduler->NetworkSemaphore.Release();
 
 		return Result;
 	}
@@ -2890,8 +2881,8 @@ SyncDirectory(const FSyncDirectoryOptions& SyncOptions)
 		std::atomic<uint64> NumBackgroundTasks = {};
 		std::atomic<uint64> NumForegroundTasks = {};
 
-		FTaskGroup BackgroundTaskGroup = GScheduler->CreateTaskGroup(EWorkloadType::Background);
-		FTaskGroup ForegroundTaskGroup = GScheduler->CreateTaskGroup(EWorkloadType::Foreground);
+		FTaskGroup BackgroundTaskGroup = GScheduler->CreateTaskGroup();
+		FTaskGroup ForegroundTaskGroup = GScheduler->CreateTaskGroup();
 
 		std::atomic<uint64> BackgroundTaskMemory	   = {};
 		std::atomic<uint64> RemainingSourceBytes	   = EstimatedNeedBytesFromSource;
@@ -2976,6 +2967,11 @@ SyncDirectory(const FSyncDirectoryOptions& SyncOptions)
 				continue;
 			}
 
+			if (GScheduler->ExecuteTasksUntilIdle())
+			{
+				continue;
+			}
+
 			std::unique_lock<std::mutex> SchedulerLock(SchedulerMutex);
 			SchedulerEvent.wait(SchedulerLock);
 		}
diff --git a/Engine/Source/Programs/Unsync/Private/UnsyncJupiter.cpp b/Engine/Source/Programs/Unsync/Private/UnsyncJupiter.cpp
--- a/Engine/Source/Programs/Unsync/Private/UnsyncJupiter.cpp
+++ b/Engine/Source/Programs/Unsync/Private/UnsyncJupiter.cpp
@@ -334,7 +334,7 @@ JupiterPutManifest(const FDirectoryManifest&	 Manifest,
 
 	for (uint64 ManifestChunkIndex = 0; ManifestChunkIndex < ChunkedRefManifestCb.size(); ++ManifestChunkIndex)
 	{
-		GScheduler->NetworkSempahore.Acquire();	 // must be acquired before task is spawned
+		GScheduler->NetworkSemaphore.Acquire();	 // must be acquired before task is spawned
 
 		ManifestPutTask&	 Context	   = TaskContexts[ManifestChunkIndex];
 		const FMiniCbWriter& RefManifestCb = ChunkedRefManifestCb[ManifestChunkIndex];
@@ -380,7 +380,7 @@ JupiterPutManifest(const FDirectoryManifest&	 Manifest,
 			}
 
 			ConnectionPool.Release(std::move(Connection));
-			GScheduler->NetworkSempahore.Release();
+			GScheduler->NetworkSemaphore.Release();
 
 			if (Response.Success())
 			{
@@ -517,7 +517,7 @@ JupiterPush(const FDirectoryManifest& Manifest, const FRemoteDesc& RemoteDesc, c
 
 	std::atomic<int32>	NumUploadedBlocks = {};
 	std::atomic<uint64> ProcessedBytes	  = {};
-	FTaskGroup			UploadTasks		  = GScheduler->CreateTaskGroup(EWorkloadType::Upload);
+	FTaskGroup			UploadTasks		  = GScheduler->CreateTaskGroup();
 	std::atomic<bool>	bGotError		  = false;
 
 	for (const auto& It : Manifest.Files)
@@ -558,7 +558,7 @@ JupiterPush(const FDirectoryManifest& Manifest, const FRemoteDesc& RemoteDesc, c
 				// Buffer compressed_buffer = compress(read_buffer.GetData(), read_buffer.GetSize());
 				// jupiter_put_compressed_blob(connection, compressed_blob_base_url, compressed_buffer, block.HashStrong);
 
-				GScheduler->NetworkSempahore.Acquire();	 // must be acquired before task is spawned
+				GScheduler->NetworkSemaphore.Acquire();	 // must be acquired before task is spawned
 				auto Task = [Block,
 							 &ProcessedBytes,
 							 &ConnectionPool,
@@ -594,7 +594,7 @@ JupiterPush(const FDirectoryManifest& Manifest, const FRemoteDesc& RemoteDesc, c
 					}
 
 					ConnectionPool.Release(std::move(Connection));
-					GScheduler->NetworkSempahore.Release();
+					GScheduler->NetworkSemaphore.Release();
 
 						if (bPutSucceeded)
 						{
diff --git a/Engine/Source/Programs/Unsync/Private/UnsyncMain.cpp b/Engine/Source/Programs/Unsync/Private/UnsyncMain.cpp
--- a/Engine/Source/Programs/Unsync/Private/UnsyncMain.cpp
+++ b/Engine/Source/Programs/Unsync/Private/UnsyncMain.cpp
@@ -654,8 +654,6 @@ InnerMain(int Argc, char** Argv)
 	UNSYNC_ASSERT(GScheduler == nullptr);
 	GScheduler = &MainScheduler;
 
-	FConcurrencyPolicyScope ConcurrencyLimitScope(GMaxThreads);
-
 	if (Cli.got_subcommand(SubHash) || Cli.got_subcommand(SubPack))
 	{
 		UNSYNC_VERBOSE(L"Using block size: %d KB", HashOrSyncBlockSize / 1024);
diff --git a/Engine/Source/Programs/Unsync/Private/UnsyncManifest.cpp b/Engine/Source/Programs/Unsync/Private/UnsyncManifest.cpp
--- a/Engine/Source/Programs/Unsync/Private/UnsyncManifest.cpp
+++ b/Engine/Source/Programs/Unsync/Private/UnsyncManifest.cpp
@@ -178,7 +178,9 @@ UpdateDirectoryManifestBlocks(FDirectoryManifest& Result, const FPath& Root, con
 
 	uint32 NumProcessedFiles = 0;
 
-	FTaskGroup TaskGroup = GScheduler->CreateTaskGroup();
+	const uint32		MaxConcurrentFiles = 8;	 // quickly diminishing returns past 8 concurrent files
+	FSchedulerSemaphore Semaphore(*GScheduler, MaxConcurrentFiles);
+	FTaskGroup			TaskGroup = GScheduler->CreateTaskGroup(&Semaphore);
 
 	uint64 NumSkippedBlocks = 0;
 	uint64 NumSkippedBytes	= 0;
@@ -200,8 +202,6 @@ UpdateDirectoryManifestBlocks(FDirectoryManifest& Result, const FPath& Root, con
 
 		FPath FilePath = Root / It.first;
 
-		GScheduler->NetworkSempahore.Acquire();
-
 		UNSYNC_VERBOSE(L"Computing blocks for '%ls' (%.2f MB)", FilePath.wstring().c_str(), SizeMb(It.second.Size));
 		auto BlockTask = [&FileManifest, &Params, FilePath = std::move(FilePath)]()
 		{
@@ -220,8 +220,6 @@ UpdateDirectoryManifestBlocks(FDirectoryManifest& Result, const FPath& Root, con
 							 FilePath.wstring().c_str(),
 							 FormatSystemErrorMessage(File.GetError()).c_str());
 			}
-
-			GScheduler->NetworkSempahore.Release();
 		};
 
 		if (Params.bAllowThreading)
@@ -286,11 +284,9 @@ CreateDirectoryManifest(const FPath& Root, const FComputeBlocksParams& Params)
 
 	FTimePoint TimeBegin = TimePointNow();
 
-	FTaskGroup TaskGroup = GScheduler->CreateTaskGroup();
-
-	const uint32 MaxConcurrentFiles = 8;  // quickly diminishing returns past 8 concurrent files
-
+	const uint32		MaxConcurrentFiles = 8;	 // quickly diminishing returns past 8 concurrent files
 	FSchedulerSemaphore Semaphore(*GScheduler, MaxConcurrentFiles);
+	FTaskGroup			TaskGroup = GScheduler->CreateTaskGroup(&Semaphore);
 
 	std::mutex ResultMutex;
 	FPath	   UnsyncDirName = ".unsync";
@@ -334,9 +330,8 @@ CreateDirectoryManifest(const FPath& Root, const FComputeBlocksParams& Params)
 			{
 				UNSYNC_VERBOSE(L"Computing blocks for '%ls' (%.2f MB)", FilePath.wstring().c_str(), double(File->GetSize()) / (1 << 20));
 
-				Semaphore.Acquire();
 				TaskGroup.run(
-					[&Semaphore, &ResultMutex, &Result, File = std::move(File), Key = std::move(PathKey), &Params]()
+					[&ResultMutex, &Result, File = std::move(File), Key = std::move(PathKey), &Params]()
 					{
 						FComputeBlocksResult ComputedBlocks = ComputeBlocks(*File, Params);
 
@@ -344,8 +339,6 @@ CreateDirectoryManifest(const FPath& Root, const FComputeBlocksParams& Params)
 
 						std::swap(Result.Files[Key].Blocks, ComputedBlocks.Blocks);
 						std::swap(Result.Files[Key].MacroBlocks, ComputedBlocks.MacroBlocks);
-
-						Semaphore.Release();
 					});
 			}
 			else
diff --git a/Engine/Source/Programs/Unsync/Private/UnsyncMemory.cpp b/Engine/Source/Programs/Unsync/Private/UnsyncMemory.cpp
--- a/Engine/Source/Programs/Unsync/Private/UnsyncMemory.cpp
+++ b/Engine/Source/Programs/Unsync/Private/UnsyncMemory.cpp
@@ -37,14 +37,19 @@ _mm_malloc(size_t s, size_t a)
 
 struct DebugHeap;
 
-namespace unsync {
+#if defined(__has_feature)
+#	if __has_feature(thread_sanitizer)
+#		define UNSYNC_THREAD_SANITIZER 1
+#	endif
+#endif
 
-#if defined(_DEBUG)
-#	define UNSYNC_OVERRIDE_GLOBAL_NEW_DELETE 0	 // TODO: need to override concurrency runtime allocator for this to work in debug config
-#else
-#	define UNSYNC_OVERRIDE_GLOBAL_NEW_DELETE 1
+#ifndef UNSYNC_THREAD_SANITIZER
+#	define UNSYNC_THREAD_SANITIZER 0
 #endif
 
+namespace unsync {
+
+#define UNSYNC_OVERRIDE_GLOBAL_NEW_DELETE	 (!UNSYNC_THREAD_SANITIZER)
 #define UNSYNC_DEBUG_MALLOC_CANARY			 1
 #define UNSYNC_DEBUG_HEAP_MINIMUM_ALLOC_SIZE 4096
 
diff --git a/Engine/Source/Programs/Unsync/Private/UnsyncProxy.cpp b/Engine/Source/Programs/Unsync/Private/UnsyncProxy.cpp
--- a/Engine/Source/Programs/Unsync/Private/UnsyncProxy.cpp
+++ b/Engine/Source/Programs/Unsync/Private/UnsyncProxy.cpp
@@ -717,7 +717,7 @@ ProxyQuery::DownloadFile(FHttpConnection&					 InConnection,
 			return;
 		}
 
-		GScheduler->NetworkSempahore.Acquire();
+		GScheduler->NetworkSemaphore.Acquire();
 
 		std::unique_ptr<FHttpConnection> Connection = ConnectionPool.Acquire();
 
@@ -747,7 +747,7 @@ ProxyQuery::DownloadFile(FHttpConnection&					 InConnection,
 
 		ConnectionPool.Release(std::move(Connection));
 
-		GScheduler->NetworkSempahore.Release();
+		GScheduler->NetworkSemaphore.Release();
 	};
 
 	ParallelForEach(Chunks, ProcessChunk);
diff --git a/Engine/Source/Programs/Unsync/Private/UnsyncScheduler.cpp b/Engine/Source/Programs/Unsync/Private/UnsyncScheduler.cpp
--- a/Engine/Source/Programs/Unsync/Private/UnsyncScheduler.cpp
+++ b/Engine/Source/Programs/Unsync/Private/UnsyncScheduler.cpp
@@ -8,35 +8,21 @@ FScheduler* GScheduler = nullptr;
 
 FScheduler::FScheduler(uint32 InNumWorkerThreads)
 : NumWorkerThreads(InNumWorkerThreads)
-, NetworkSempahore(*this, MAX_NETWORK_TASKS)
-, FilesystemSemaphore(*this, MAX_FILESYSTEM_TASKS)
+, NetworkSemaphore(*this, MAX_NETWORK_TASKS)
 {
-#if UNSYNC_USE_CONCRT
-	UNSYNC_VERBOSE(L"Scheduler: PPL");
-#else
-	UNSYNC_VERBOSE(L"Scheduler: Custom");
 	ThreadPool.StartWorkers(NumWorkerThreads);
-#endif
 }
 
 FScheduler::~FScheduler()
 {
 }
 
-FTaskGroup FScheduler::CreateTaskGroup(EWorkloadType Type)
+FTaskGroup
+FScheduler::CreateTaskGroup(FSchedulerSemaphore* ConcurrencyLimiter)
 {
-	// TODO: use separate thread pools for tasks of different types
-	UNSYNC_UNUSED(Type);
-
-#if UNSYNC_USE_CONCRT
-	return FTaskGroup();
-#else
-	return FTaskGroup(ThreadPool);
-#endif
+	return FTaskGroup(ThreadPool, ConcurrencyLimiter);
 }
 
-#if !UNSYNC_USE_CONCRT
-
 FSchedulerSemaphore::FSchedulerSemaphore(FScheduler& InScheduler, uint32 MaxCount)
 : Scheduler(InScheduler)
 , Native(std::min(MaxCount, 1 + InScheduler.NumWorkerThreads)) // 1 for main thread + N workers
@@ -44,11 +30,18 @@ FSchedulerSemaphore::FSchedulerSemaphore(FScheduler& InScheduler, uint32 MaxCoun
 }
 
 void
-FSchedulerSemaphore::Acquire()
+FSchedulerSemaphore::Acquire(bool bAllowTaskExecution)
 {
-	while (!Native.try_acquire())
+	if (bAllowTaskExecution)
+	{
+		while (!Native.try_acquire())
+		{
+			Scheduler.TryExecuteTask();
+		}
+	}
+	else
 	{
-		Scheduler.TryExecuteTask();
+		Native.acquire();
 	}
 }
 
@@ -58,6 +51,4 @@ FSchedulerSemaphore::Release()
 	Native.release();
 }
 
-#endif
-
 }  // namespace unsync
diff --git a/Engine/Source/Programs/Unsync/Private/UnsyncScheduler.h b/Engine/Source/Programs/Unsync/Private/UnsyncScheduler.h
--- a/Engine/Source/Programs/Unsync/Private/UnsyncScheduler.h
+++ b/Engine/Source/Programs/Unsync/Private/UnsyncScheduler.h
@@ -12,80 +12,54 @@ struct FTaskGroup;
 
 extern FScheduler* GScheduler;
 
-enum class EWorkloadType : uint8
-{
-	Foreground,
-	Background,
-	Download,
-	Upload,
-	FileIO,
-};
-
-#if UNSYNC_USE_CONCRT
-struct FSchedulerSemaphore : FNativeSemaphore
-{
-	explicit FSchedulerSemaphore(FScheduler& InScheduler, uint32 MaxCount) : FNativeSemaphore(MaxCount) { UNSYNC_UNUSED(InScheduler); }
-};
-#else
 struct FSchedulerSemaphore
 {
 	UNSYNC_DISALLOW_COPY_ASSIGN(FSchedulerSemaphore)
 
 	FSchedulerSemaphore(FScheduler& InScheduler, uint32 MaxCount);
 
-	void Acquire();
+	bool TryAcquire() { return Native.try_acquire(); }
+
+	void Acquire(bool bAllowTaskExecution = true);
 	void Release();
 
 	FScheduler& Scheduler;
 
 	std::counting_semaphore<UNSYNC_MAX_TOTAL_THREADS> Native;
 };
-#endif
 
 class FScheduler
 {
 public:
 	UNSYNC_DISALLOW_COPY_ASSIGN(FScheduler)
 
-	static constexpr uint32 MAX_NETWORK_TASKS	 = 8;
-	static constexpr uint32 MAX_FILESYSTEM_TASKS = 64;
+	static constexpr uint32 MAX_NETWORK_TASKS = 8;
 
 	FScheduler(uint32 InNumWorkerThreads);
 	~FScheduler();
 
 	const uint32 NumWorkerThreads;
 
-	FSchedulerSemaphore NetworkSempahore;
-	FSchedulerSemaphore FilesystemSemaphore;
+	FSchedulerSemaphore NetworkSemaphore;
 
-	FTaskGroup CreateTaskGroup(EWorkloadType Type = EWorkloadType::Foreground);
+	FTaskGroup CreateTaskGroup(FSchedulerSemaphore* ConcurrencyLimiter = nullptr);
 
 	void TryExecuteTask() { ThreadPool.TryExecuteTask(); }
 
-private:
-	FThreadPool ThreadPool;
-};
-
-#if UNSYNC_USE_CONCRT
-
-struct FTaskGroup : concurrency::task_group
-{
-	UNSYNC_DISALLOW_COPY_ASSIGN(FTaskGroup)
+	bool ExecuteTasksUntilIdle()
+	{
+		uint64 NumExecuted = 0;
+		while (ThreadPool.TryExecuteTask())
+		{
+			NumExecuted++;
+		}
+		return NumExecuted != 0;
+	}
 
 private:
-	friend FScheduler;
-	FTaskGroup() : concurrency::task_group() {}
+	FThreadPool ThreadPool;
 };
 
-template<typename IT, typename FT>
-inline void
-ParallelForEach(IT ItBegin, IT ItEnd, FT F)
-{
-	concurrency::parallel_for_each(ItBegin, ItEnd, F);
-}
-
-#else // UNSYNC_USE_CONCRT
-
 struct FTaskGroup
 {
 	UNSYNC_DISALLOW_COPY_ASSIGN(FTaskGroup)
@@ -95,14 +69,30 @@ struct FTaskGroup
 	{
 		++NumStartedTasks;
 
-		ThreadPool.PushTask(
-			[&NumStartedTasks  = this->NumStartedTasks,
-			 &NumFinishedTasks = this->NumFinishedTasks,
-			 Function		   = std::forward<F>(InFunction)]() -> void
-			{
-				Function();
-				++NumFinishedTasks;
-			});
+		const bool bAcquired = Semaphore && Semaphore->TryAcquire();
+
+		if (!Semaphore || bAcquired)
+		{
+			ThreadPool.PushTask(
+				[Semaphore = this->Semaphore,
+				 bAcquired,
+				 &NumStartedTasks  = this->NumStartedTasks,
+				 &NumFinishedTasks = this->NumFinishedTasks,
+				 Function		   = std::forward<F>(InFunction)]() -> void
+				{
+					Function();
+					if (Semaphore && bAcquired)
+					{
+						Semaphore->Release();
+					}
+					++NumFinishedTasks;
+				});
+		}
+		else
+		{
+			InFunction();
+			++NumFinishedTasks;
+		}
 	}
 
 	void wait()
@@ -116,12 +106,13 @@ struct FTaskGroup
 	FThreadPool&		ThreadPool;
 	std::atomic<uint64> NumStartedTasks;
 	std::atomic<uint64> NumFinishedTasks;
+	FSchedulerSemaphore* Semaphore = nullptr;
 
 	~FTaskGroup() { wait(); };
 
 private:
 	friend FScheduler;
-	FTaskGroup(FThreadPool& InThreadPool) : ThreadPool(InThreadPool) {}
+	FTaskGroup(FThreadPool& InThreadPool, FSchedulerSemaphore* InSemaphore) : ThreadPool(InThreadPool), Semaphore(InSemaphore) {}
 };
 
 template<typename IT, typename FT>
@@ -139,8 +130,6 @@ ParallelForEach(IT ItBegin, IT ItEnd, FT F)
 	TaskGroup.wait();
 }
 
-#endif // UNSYNC_USE_CONCRT
-
 template<typename T, typename FT>
 inline void
 ParallelForEach(T& Container, FT F)
diff --git a/Engine/Source/Programs/Unsync/Private/UnsyncTarget.cpp b/Engine/Source/Programs/Unsync/Private/UnsyncTarget.cpp
--- a/Engine/Source/Programs/Unsync/Private/UnsyncTarget.cpp
+++ b/Engine/Source/Programs/Unsync/Private/UnsyncTarget.cpp
@@ -81,12 +81,11 @@ DownloadBlocks(FProxyPool&					  ProxyPool,
 
 		UNSYNC_VERBOSE2(L"Download batches: %lld", Batches.size());
 
-		FTaskGroup DownloadTasks = GScheduler->CreateTaskGroup();
+		FTaskGroup DownloadTasks = GScheduler->CreateTaskGroup(&GScheduler->NetworkSemaphore);
 		std::mutex DownloadedBlocksMutex;
 
 		for (FDownloadBatch Batch : Batches)
 		{
-			GScheduler->NetworkSempahore.Acquire();
 			DownloadTasks.run(
 				[NeedBlocks,
 				 Batch,
@@ -101,7 +100,6 @@ DownloadBlocks(FProxyPool&					  ProxyPool,
 				{
 					if (bGotError)
 					{
-						GScheduler->NetworkSempahore.Release();
 						return;
 					}
 
@@ -137,7 +135,6 @@ DownloadBlocks(FProxyPool&					  ProxyPool,
 							ProxyPool.Invalidate();
 						}
 					}
-					GScheduler->NetworkSempahore.Release();
 				});
 		}
 
@@ -196,7 +193,10 @@ BuildTarget(FIOWriter& Output, FIOReader& Source, FIOReader& Base, const FNeedLi
 	};
 	FStats Stats;
 
-	FTaskGroup WriteTasks = GScheduler->CreateTaskGroup();
+	const uint32 MaxWriteTasks = 64;
+
+	FSchedulerSemaphore WriteSemaphore(*GScheduler, MaxWriteTasks); // // throttle writing tasks to avoid memory bloat
+	FTaskGroup			WriteTasks = GScheduler->CreateTaskGroup(&WriteSemaphore);
 
 	// Remember if parent thread has verbose logging and indentation
 	const bool	 bAllowVerboseLog = GLogVerbose;
@@ -331,7 +331,6 @@ BuildTarget(FIOWriter& Output, FIOReader& Source, FIOReader& Base, const FNeedLi
 																													 uint64	   CmdReadSize,
 																													 uint64	   CmdUserData)
 			{
-				GScheduler->FilesystemSemaphore.Acquire();
 				WriteTasks.run(
 					[Buffer = MakeShared(std::move(CmdBuffer)), CmdReadSize, Block, &Output, &Error, &Stats, ListType]()
 					{
@@ -350,8 +349,6 @@ BuildTarget(FIOWriter& Output, FIOReader& Source, FIOReader& Base, const FNeedLi
 							UNSYNC_FATAL(L"Unexpected block list type");
 						}
 
-						GScheduler->FilesystemSemaphore.Release();
-
 						if (WrittenBytes != CmdReadSize)
 						{
 							UNSYNC_FATAL(L"Expected to write %llu bytes, but written %llu", CmdReadSize, WrittenBytes);
@@ -381,8 +378,6 @@ BuildTarget(FIOWriter& Output, FIOReader& Source, FIOReader& Base, const FNeedLi
 											  (bAllowVerboseLog && (ListType == EBlockListType::Base && bWaitingForBaseData)));
 
 			ProgressLogger.Add(ReadBytes);
-
-			SchedulerYield();
 		}
 
 		DataProvider.FlushAll();
@@ -477,13 +472,13 @@ BuildTarget(FIOWriter& Output, FIOReader& Source, FIOReader& Base, const FNeedLi
 			}
 		}
 
-		FTaskGroup DecompressTasks = GScheduler->CreateTaskGroup();
+		// limit how many decompression tasks can be queued up to avoid memory bloat
+		const uint64		MaxConcurrentDecompressionTasks = 64;
+		FSchedulerSemaphore DecompressionSemaphore(*GScheduler, MaxConcurrentDecompressionTasks);
 
-		FLogProgressScope DownloadProgressLogger(EstimatedDownloadSize, ELogProgressUnits::MB);
+		FTaskGroup DecompressTasks = GScheduler->CreateTaskGroup(&DecompressionSemaphore);
 
-		// limit how many decompression tasks can be queued up to avoid memory bloat
-		const uint64 MaxConcurrentDecompressionTasks = 64;
-		FSchedulerSemaphore	 DecompressionSemaphore(*GScheduler, MaxConcurrentDecompressionTasks);
+		FLogProgressScope DownloadProgressLogger(EstimatedDownloadSize, ELogProgressUnits::MB);
 
 		const bool			bParentThreadVerbose = GLogVerbose;
 		const uint32		ParentThreadIndent	 = GLogIndent;
@@ -529,8 +524,6 @@ BuildTarget(FIOWriter& Output, FIOReader& Source, FIOReader& Base, const FNeedLi
 
 					memcpy(DownloadedData.GetData(), Block.Data, DownloadedSize);
 
-					DecompressionSemaphore.Acquire();
-
 					DecompressTasks.run(
 						[&Output,
 						 BlockHashUnaligned = BlockHash,
@@ -543,7 +536,6 @@ BuildTarget(FIOWriter& Output, FIOReader& Source, FIOReader& Base, const FNeedLi
 						 &BlockScatterMap,
 						 &DownloadedBlocksMutex,
 						 &DownloadedBlocks,
-						 &DecompressionSemaphore,
 						 &Error,
 						 &NumHashMismatches,
 						 &Stats]()
@@ -634,8 +626,6 @@ BuildTarget(FIOWriter& Output, FIOReader& Source, FIOReader& Base, const FNeedLi
 								std::lock_guard<std::mutex> LockGuard(DownloadedBlocksMutex);
 								DownloadedBlocks.insert(BlockHash);
 							}
-
-							DecompressionSemaphore.Release();
 						});
 				}
 			});
diff --git a/Engine/Source/Programs/Unsync/Private/UnsyncTest.cpp b/Engine/Source/Programs/Unsync/Private/UnsyncTest.cpp
--- a/Engine/Source/Programs/Unsync/Private/UnsyncTest.cpp
+++ b/Engine/Source/Programs/Unsync/Private/UnsyncTest.cpp
@@ -895,8 +895,6 @@ RunTests(const std::string& Preset)
 
 	if (Preset == "perf" || Preset == "all")
 	{
-		FConcurrencyPolicyScope SingleThreadedScope(1);
-
 		TestPerfComputeBlocksVariable(EWeakHashAlgorithmID::Naive, EStrongHashAlgorithmID::Blake3_128);
 		TestPerfComputeBlocksVariable(EWeakHashAlgorithmID::BuzHash, EStrongHashAlgorithmID::Blake3_128);
 		TestPerfComputeBlocksVariable(EWeakHashAlgorithmID::BuzHash, EStrongHashAlgorithmID::Blake3_160);
diff --git a/Engine/Source/Programs/Unsync/Private/UnsyncThread.cpp b/Engine/Source/Programs/Unsync/Private/UnsyncThread.cpp
--- a/Engine/Source/Programs/Unsync/Private/UnsyncThread.cpp
+++ b/Engine/Source/Programs/Unsync/Private/UnsyncThread.cpp
@@ -52,8 +52,8 @@ FThreadPool::PopTask(bool bWaitForSignal)
 
 	if (!Tasks.empty())
 	{
-		Result = std::move(Tasks.front());
-		Tasks.pop_front();
+		Result = std::move(Tasks.back());
+		Tasks.pop_back();
 	}
 
 	return Result;
@@ -95,53 +95,6 @@ FThreadPool::DoWorkInternal(bool bWaitForSignal)
 	}
 }
 
-#if UNSYNC_USE_CONCRT
-FConcurrencyPolicyScope::FConcurrencyPolicyScope(uint32 MaxConcurrency)
-{
-	auto Policy = Concurrency::CurrentScheduler::GetPolicy();
-
-	const uint32 CurrentMaxConcurrency = Policy.GetPolicyValue(Concurrency::PolicyElementKey::MaxConcurrency);
-	const uint32 CurrentMinConcurrency = Policy.GetPolicyValue(Concurrency::PolicyElementKey::MinConcurrency);
-
-	MaxConcurrency = std::min(MaxConcurrency, std::thread::hardware_concurrency());
-	MaxConcurrency = std::min(MaxConcurrency, CurrentMaxConcurrency);
-	MaxConcurrency = std::max(1u, MaxConcurrency);
-
-	Policy.SetConcurrencyLimits(CurrentMinConcurrency, MaxConcurrency);
-
-	Concurrency::CurrentScheduler::Create(Policy);
-}
-
-FConcurrencyPolicyScope::~FConcurrencyPolicyScope()
-{
-	Concurrency::CurrentScheduler::Detach();
-}
-
-void
-SchedulerSleep(uint32 Milliseconds)
-{
-	concurrency::event E;
-	E.reset();
-	E.wait(Milliseconds);
-}
-
-void
-SchedulerYield()
-{
-	concurrency::Context::YieldExecution();
-}
-
-#else  // UNSYNC_USE_CONCRT
-
-FConcurrencyPolicyScope::FConcurrencyPolicyScope(uint32 MaxConcurrency)
-{
-	// TODO
-}
-
-FConcurrencyPolicyScope::~FConcurrencyPolicyScope()
-{
-	// TODO
-}
 
 void
 SchedulerSleep(uint32 Milliseconds)
@@ -149,14 +102,6 @@ SchedulerSleep(uint32 Milliseconds)
 	std::this_thread::sleep_for(std::chrono::milliseconds(Milliseconds));
 }
 
-void
-SchedulerYield()
-{
-	// TODO
-}
-
-#endif	// UNSYNC_USE_CONCRT
-
 void
 TestThread()
 {
diff --git a/Engine/Source/Programs/Unsync/Private/UnsyncThread.h b/Engine/Source/Programs/Unsync/Private/UnsyncThread.h
--- a/Engine/Source/Programs/Unsync/Private/UnsyncThread.h
+++ b/Engine/Source/Programs/Unsync/Private/UnsyncThread.h
@@ -14,28 +14,7 @@ UNSYNC_THIRD_PARTY_INCLUDES_START
 #include <functional>
 #include <thread>
 #include <vector>
-#if UNSYNC_USE_CONCRT
-#	include <concrt.h>
-#	include <concurrent_queue.h>
-#	include <ppl.h>
-#	ifdef Yield
-#		undef Yield  // WinBase.h defines this :-/
-#	endif
-#else
-#	include <semaphore>
-#	ifdef __APPLE__
-#		define UNSYNC_USE_MACH_SEMAPHORE 1
-#	endif
-#endif	// UNSYNC_USE_CONCRT
-#if !defined(UNSYNC_USE_MACH_SEMAPHORE)
-#	define UNSYNC_USE_MACH_SEMAPHORE 0
-#endif
-#if UNSYNC_USE_MACH_SEMAPHORE
-#include <mach/mach_init.h>
-#include <mach/mach_error.h>
-#include <mach/semaphore.h>
-#include <mach/task.h>
-#endif // UNSYNC_USE_MACH_SEMAPHORE
+#include <semaphore>
 UNSYNC_THIRD_PARTY_INCLUDES_END
 
 namespace unsync {
@@ -46,14 +25,6 @@ extern uint32 GMaxThreads;
 
 class FScheduler;
 
-class FConcurrencyPolicyScope
-{
-public:
-	UNSYNC_DISALLOW_COPY_ASSIGN(FConcurrencyPolicyScope)
-	explicit FConcurrencyPolicyScope(uint32 MaxConcurrency);
-	~FConcurrencyPolicyScope();
-};
-
 struct FThreadElectScope
 {
 	const bool			 bValue;	  // NOLINT
@@ -100,13 +71,13 @@ struct FThreadLogConfig
 };
 
 void SchedulerSleep(uint32 Milliseconds);
-void SchedulerYield();
 
 class FThreadPool
 {
 public:
 
-	using FTaskFunction = std::function<void()>;
+	using FTaskFunction = std::function<
+		void()>;
 
 	FThreadPool() = default;
 	~FThreadPool();
@@ -132,7 +103,7 @@ private:
 	FTaskFunction PopTask(bool bWaitForSignal);
 
 	std::vector<std::thread>  Threads;
-	std::deque<FTaskFunction> Tasks;
+	std::vector<FTaskFunction> Tasks;
 
 	std::mutex				Mutex;
 	std::condition_variable WorkerWakeupCondition;
@@ -140,105 +111,4 @@ private:
 	std::atomic<uint64>		NumRunningTasks;
 };
 
-#if UNSYNC_USE_CONCRT
-
-// Cooperative semaphore implementation.
-// Using this is necessary to avoid deadlocks on low-core machines.
-// https://docs.microsoft.com/en-us/cpp/parallel/concrt/how-to-use-the-context-class-to-implement-a-cooperative-semaphore?view=msvc-160
-class FNativeSemaphore
-{
-public:
-	UNSYNC_DISALLOW_COPY_ASSIGN(FNativeSemaphore)
-
-	explicit FNativeSemaphore(uint32 MaxCount) : Counter(MaxCount) {}
-
-	~FNativeSemaphore() {}
-
-	void Acquire()
-	{
-		if (--Counter < 0)
-		{
-			WaitingQueue.push(concurrency::Context::CurrentContext());
-			concurrency::Context::Block();
-		}
-	}
-
-	void Release()
-	{
-		if (++Counter <= 0)
-		{
-			concurrency::Context* Waiting = nullptr;
-			while (!WaitingQueue.try_pop(Waiting))
-			{
-				concurrency::Context::YieldExecution();
-			}
-			Waiting->Unblock();
-		}
-	}
-
-private:
-	std::atomic<int64>									 Counter;
-	concurrency::concurrent_queue<concurrency::Context*> WaitingQueue;
-};
-
-#else  // UNSYNC_USE_CONCRT
-
-#if UNSYNC_USE_MACH_SEMAPHORE
-
-class FNativeSemaphore
-{
-public:
-	UNSYNC_DISALLOW_COPY_ASSIGN(FNativeSemaphore)
-
-	explicit FNativeSemaphore(uint32 MaxCount)
-	{
-		kern_return_t InitResult = semaphore_create(mach_task_self(), &Native, SYNC_POLICY_FIFO, MaxCount);
-		UNSYNC_ASSERTF(InitResult == KERN_SUCCESS, L"Failed to create a semaphore, error code: %d %hs", InitResult, mach_error_string(InitResult));
-	}
-
-	~FNativeSemaphore()
-	{
-		kern_return_t DestroyResult = semaphore_destroy(mach_task_self(), Native);
-		UNSYNC_ASSERTF(DestroyResult == KERN_SUCCESS, L"Failed to destroy a semaphore, error code: %d %hs", DestroyResult, mach_error_string(DestroyResult));
-	}
-
-	void Acquire()
-	{
-		kern_return_t WaitResult = semaphore_wait(Native);
-		UNSYNC_ASSERTF(WaitResult == KERN_SUCCESS, L"Failed to wait for a semaphore, error code: %d %hs", WaitResult, mach_error_string(WaitResult));
-	}
-
-	void Release()
-	{
-		int32 SignalResult = semaphore_signal(Native);
-		UNSYNC_ASSERTF(SignalResult == 0, L"Failed to signal a semaphore, error code: %d %hs", SignalResult, mach_error_string(SignalResult));
-	}
-
-	semaphore_t Native = {};
-};
-
-#else // UNSYNC_USE_MACH_SEMAPHORE
-
-class FNativeSemaphore
-{
-public:
-	UNSYNC_DISALLOW_COPY_ASSIGN(FNativeSemaphore)
-
-	explicit FNativeSemaphore(uint32 MaxCount) : Native(MaxCount) {}
-
-	~FNativeSemaphore() {}
-
-	void Acquire() { Native.acquire(); }
-
-	void Release() { Native.release(); }
-
-private:
-
-	std::counting_semaphore<UNSYNC_MAX_TOTAL_THREADS> Native;
-};
-
-#endif // UNSYNC_USE_MACH_SEMAPHORE
-
-#endif	// UNSYNC_USE_CONCRT
-
 }  // namespace unsync
diff --git a/Engine/Source/Programs/Unsync/Unsync.Build.cs b/Engine/Source/Programs/Unsync/Unsync.Build.cs
--- a/Engine/Source/Programs/Unsync/Unsync.Build.cs
+++ b/Engine/Source/Programs/Unsync/Unsync.Build.cs
@@ -19,7 +19,6 @@ public class Unsync : ModuleRules
 		if (Target.Platform == UnrealTargetPlatform.Win64)
 		{
 			PrivateDefinitions.Add("UNSYNC_PLATFORM_WINDOWS=1");
-			PrivateDefinitions.Add("UNSYNC_USE_CONCRT=1");
 			PrivateDefinitions.Add("UNSYNC_PLATFORM_UNIX=0");
 			PrivateDefinitions.Add("NOMINMAX=1");
 			PrivateDefinitions.Add("WIN32_LEAN_AND_MEAN=1");
@@ -31,10 +30,14 @@ public class Unsync : ModuleRules
 		else if (Target.Platform == UnrealTargetPlatform.Mac || Target.Platform == UnrealTargetPlatform.Linux)
 		{
 			PrivateDefinitions.Add("UNSYNC_PLATFORM_WINDOWS=0");
-			PrivateDefinitions.Add("UNSYNC_USE_CONCRT=0");
 			PrivateDefinitions.Add("UNSYNC_PLATFORM_UNIX=1");
 		}
 
+		if (Target.Platform == UnrealTargetPlatform.Linux)
+		{
+			PublicSystemLibraries.Add("dl");
+		}
+
 		PrivateDependencyModuleNames.Add("BLAKE3");
 		PrivateDependencyModuleNames.Add("CLI11");
 		PrivateDependencyModuleNames.Add("fmt");
