[Unsync] Fix null termination issue in memory buffer of decoded login token that was causing garbage bytes to be printed.

#rnx

diff --git a/Engine/Source/Programs/Unsync/Private/UnsyncBuffer.h b/Engine/Source/Programs/Unsync/Private/UnsyncBuffer.h
--- a/Engine/Source/Programs/Unsync/Private/UnsyncBuffer.h
+++ b/Engine/Source/Programs/Unsync/Private/UnsyncBuffer.h
@@ -32,7 +32,7 @@ class FBuffer
 {
 public:
 	FBuffer() = default;
-	FBuffer(uint64 InitialSize);
+	explicit FBuffer(uint64 InitialSize);
 	FBuffer(uint64 InitialSize, uint8 Value);
 	FBuffer(std::initializer_list<uint8> Values);
 	FBuffer(const FBuffer&) = delete;
diff --git a/Engine/Source/Programs/Unsync/Private/UnsyncCmdLogin.cpp b/Engine/Source/Programs/Unsync/Private/UnsyncCmdLogin.cpp
--- a/Engine/Source/Programs/Unsync/Private/UnsyncCmdLogin.cpp
+++ b/Engine/Source/Programs/Unsync/Private/UnsyncCmdLogin.cpp
@@ -84,7 +84,7 @@ CmdLogin(const FCmdLoginOptions& Options)
 					UNSYNC_ERROR("Failed to Base64-decode access token");
 					return -1;
 				}
-				DecodedTokenData.Append(0);
+				DecodedTokenData.PushBack(0);
 
 				LogPrintf(ELogLevel::MachineReadable, L"%hs\n", (const char*)DecodedTokenData.Data());
 			}
