unsync - Move file open operation to worker threads during manifest generation

diff --git a/Engine/Source/Programs/Unsync/Private/UnsyncManifest.cpp b/Engine/Source/Programs/Unsync/Private/UnsyncManifest.cpp
--- a/Engine/Source/Programs/Unsync/Private/UnsyncManifest.cpp
+++ b/Engine/Source/Programs/Unsync/Private/UnsyncManifest.cpp
@@ -291,6 +291,8 @@ CreateDirectoryManifest(const FPath& Root, const FComputeBlocksParams& Params)
 	std::mutex ResultMutex;
 	FPath	   UnsyncDirName = ".unsync";
 
+	FThreadLogConfig MainLogConfig;
+
 	for (const std::filesystem::directory_entry& Dir : RecursiveDirectoryScan(Root))
 	{
 		if (Dir.is_directory())
@@ -325,28 +327,36 @@ CreateDirectoryManifest(const FPath& Root, const FComputeBlocksParams& Params)
 		if (Params.bNeedBlocks && Params.BlockSize)
 		{
 			FPath FilePath = Root / RelativePath;
-			auto  File	   = std::make_shared<FNativeFile>(FilePath, EFileMode::ReadOnlyUnbuffered);
-			if (File->IsValid())
-			{
-				UNSYNC_VERBOSE(L"Computing blocks for '%ls' (%.2f MB)", FilePath.wstring().c_str(), double(File->GetSize()) / (1 << 20));
 
-				TaskGroup.run(
-					[&ResultMutex, &Result, File = std::move(File), Key = std::move(PathKey), &Params]()
+			TaskGroup.run(
+				[&ResultMutex, &Result, &MainLogConfig, FilePath = std::move(FilePath), Key = std::move(PathKey), &Params]()
+				{
+					auto File = std::make_shared<FNativeFile>(FilePath, EFileMode::ReadOnlyUnbuffered);
+
+					if (File->IsValid())
 					{
+						{
+							// Log from worker using parent thread log config
+							FThreadLogConfig::FScope LogConfigScope(MainLogConfig);
+							UNSYNC_VERBOSE(L"Computing blocks for '%ls' (%.2f MB)",
+										   FilePath.wstring().c_str(),
+										   double(File->GetSize()) / (1 << 20));
+						}
+
 						FComputeBlocksResult ComputedBlocks = ComputeBlocks(*File, Params);
 
 						std::lock_guard<std::mutex> LockGuard(ResultMutex);
 
 						std::swap(Result.Files[Key].Blocks, ComputedBlocks.Blocks);
 						std::swap(Result.Files[Key].MacroBlocks, ComputedBlocks.MacroBlocks);
-					});
-			}
-			else
-			{
-				UNSYNC_FATAL(L"Failed to open file '%ls' while computing manifest blocks. %hs",
-							 FilePath.wstring().c_str(),
-							 FormatSystemErrorMessage(File->GetError()).c_str());
-			}
+					}
+					else
+					{
+						UNSYNC_FATAL(L"Failed to open file '%ls' while computing manifest blocks. %hs",
+										FilePath.wstring().c_str(),
+										FormatSystemErrorMessage(File->GetError()).c_str());
+					}
+				});
 		}
 	}
 
