unsync - Groundwork for packing small files into a custom container during manifest generation

diff --git a/Engine/Source/Programs/Unsync/CMakeLists.txt b/Engine/Source/Programs/Unsync/CMakeLists.txt
--- a/Engine/Source/Programs/Unsync/CMakeLists.txt
+++ b/Engine/Source/Programs/Unsync/CMakeLists.txt
@@ -78,6 +78,8 @@ set(src
 	Private/UnsyncMiniCb.h
 	Private/UnsyncMount.cpp
 	Private/UnsyncMount.h
+	Private/UnsyncPack.cpp
+	Private/UnsyncPack.h
 	Private/UnsyncPool.h
 	Private/UnsyncProgress.cpp
 	Private/UnsyncProgress.h
diff --git a/Engine/Source/Programs/Unsync/Private/UnsyncCmdHash.cpp b/Engine/Source/Programs/Unsync/Private/UnsyncCmdHash.cpp
--- a/Engine/Source/Programs/Unsync/Private/UnsyncCmdHash.cpp
+++ b/Engine/Source/Programs/Unsync/Private/UnsyncCmdHash.cpp
@@ -4,6 +4,9 @@
 #include "UnsyncCore.h"
 #include "UnsyncFile.h"
 #include "UnsyncSerialization.h"
+#include "UnsyncPack.h"
+
+#include <memory>
 
 namespace unsync {
 
@@ -42,6 +45,52 @@ CmdHash(const FCmdHashOptions& Options)
 		// TODO: macro block generation is only implemented for variable chunk mode
 		ComputeBlocksParams.bNeedMacroBlocks = ComputeBlocksParams.Algorithm.ChunkingAlgorithmId == EChunkingAlgorithmID::VariableBlocks;
 
+		std::unique_ptr<FPackWriteContext> PackWriter;
+		THashSet<FGenericHash>			   PackedBlocks;
+		std::mutex						   PackedBlocksMutex;
+
+		if (!GDryRun && Options.bPackFiles)
+		{
+			FPath PackOutputRoot = ManifestRoot / "pack";
+
+			if (!PathExists(PackOutputRoot))
+			{
+				bool bPackOutputExists = EnsureDirectoryExists(PackOutputRoot);
+				if (!bPackOutputExists)
+				{
+					UNSYNC_ERROR(L"Failed to create pack output directory '%ls'", PackOutputRoot.wstring().c_str());
+					return 1;
+				}
+			}
+
+			PackWriter = std::unique_ptr<FPackWriteContext>(new FPackWriteContext(PackOutputRoot));
+
+			UNSYNC_LOG(L"Pack output: '%ls'", PackOutputRoot.wstring().c_str());
+
+			auto OnBlockGenerated = [&PackWriter, &Options, &PackedBlocks, &PackedBlocksMutex](const FGenericBlock&	   Block,
+																							   const FBlockSourceInfo& Source,
+																							   FBufferView			   Data)
+			{
+				if (Source.TotalSize > Options.MaxFileSizeToPack)
+				{
+					return;
+				}
+
+				// Only store unique blocks in the pack
+				{
+					std::unique_lock<std::mutex> LockScope(PackedBlocksMutex);
+					if (!PackedBlocks.insert(Block.HashStrong).second)
+					{
+						return;
+					}
+				}
+
+				PackWriter->CompressAndAddBlock(Block, Data);
+			};
+
+			ComputeBlocksParams.OnBlockGenerated = OnBlockGenerated;
+		}
+
 		FDirectoryManifest DirectoryManifest;
 		if (Options.bForce)
 		{
@@ -57,6 +106,12 @@ CmdHash(const FCmdHashOptions& Options)
 			LoadOrCreateDirectoryManifest(DirectoryManifest, Options.Input, ComputeBlocksParams);
 		}
 
+		if (PackWriter)
+		{
+			PackWriter->FinishPack();
+			PackWriter->GetUniqueGeneratedPackIds(DirectoryManifest.PackReferences);
+		}
+
 		if (!GDryRun)
 		{
 			UNSYNC_LOG(L"Saving directory manifest '%ls'", DirectoryManifestPath.wstring().c_str());
@@ -114,3 +169,4 @@ CmdHash(const FCmdHashOptions& Options)
 }
 
 }  // namespace unsync
+
diff --git a/Engine/Source/Programs/Unsync/Private/UnsyncCmdHash.h b/Engine/Source/Programs/Unsync/Private/UnsyncCmdHash.h
--- a/Engine/Source/Programs/Unsync/Private/UnsyncCmdHash.h
+++ b/Engine/Source/Programs/Unsync/Private/UnsyncCmdHash.h
@@ -11,9 +11,12 @@ struct FCmdHashOptions
 {
 	FPath			  Input;
 	FPath			  Output;
-	bool			  bForce	   = false;
-	bool			  bIncremental = false;
-	uint32			  BlockSize	   = uint32(64_KB);
+	bool			  bForce			= false;
+	bool			  bIncremental		= false;
+	bool			  bCompressPacks	= false;
+	bool			  bPackFiles		= false;
+	uint64			  MaxFileSizeToPack = ~0ull;
+	uint32			  BlockSize			= uint32(64_KB);
 	FAlgorithmOptions Algorithm;
 };
 
diff --git a/Engine/Source/Programs/Unsync/Private/UnsyncCmdPack.cpp b/Engine/Source/Programs/Unsync/Private/UnsyncCmdPack.cpp
--- a/Engine/Source/Programs/Unsync/Private/UnsyncCmdPack.cpp
+++ b/Engine/Source/Programs/Unsync/Private/UnsyncCmdPack.cpp
@@ -9,6 +9,7 @@
 #include "UnsyncError.h"
 #include "UnsyncProgress.h"
 #include "UnsyncScheduler.h"
+#include "UnsyncPack.h"
 
 #include <stdio.h>
 #include <stdlib.h>
@@ -16,8 +17,6 @@
 
 namespace unsync {
 
-static constexpr uint64 GMaxPackFileSize = 1_GB;
-
 template<typename CallbackT>
 static void
 ForLines(std::string_view String, CallbackT Callback)
@@ -87,15 +86,6 @@ BuildP4HaveSet(const FPath& Root, std::string_view P4HaveDataUtf8, FDirectoryMan
 	ForLines(P4HaveDataUtf8, Callback);
 }
 
-struct FPackIndexEntry	// structure is serialized
-{
-	FHash128 BlockHash		= {};
-	FHash128 CompressedHash = {};
-	uint32	 Offset			= 0;
-	uint32	 CompressedSize = 0;
-};
-static_assert(sizeof(FPackIndexEntry) == 40);
-
 struct FPackDatabase
 {
 	struct FEntry
@@ -140,7 +130,7 @@ struct FPackDatabase
 
 			for (const FPackIndexEntry& IndexEntry : ReinterpretView<FPackIndexEntry>(IndexEntries))
 			{
-				UNSYNC_ASSERT(IndexEntry.Offset + IndexEntry.CompressedSize < GMaxPackFileSize);
+				UNSYNC_ASSERT(IndexEntry.PackBlockOffset + IndexEntry.PackBlockSize < GMaxPackFileSize);
 
 				FEntry DatabaseEntry;
 				DatabaseEntry.IndexEntry	   = IndexEntry;
@@ -148,7 +138,7 @@ struct FPackDatabase
 
 				if (BlockMap.insert(std::make_pair(IndexEntry.BlockHash, DatabaseEntry)).second)
 				{
-					TotalCompressedSize += DatabaseEntry.IndexEntry.CompressedSize;
+					TotalCompressedSize += DatabaseEntry.IndexEntry.PackBlockSize;
 				}
 			}
 		}
@@ -176,112 +166,6 @@ struct FPackDatabase
 	mutable std::mutex FileCacheMutex;
 };
 
-inline void
-AddHash(uint64* Accumulator, const FHash128& Hash)
-{
-	uint64 BlockHashParts[2];
-	memcpy(BlockHashParts, &Hash, sizeof(FHash128));
-	Accumulator[0] += BlockHashParts[0];
-	Accumulator[1] += BlockHashParts[1];
-}
-
-inline FHash128
-MakeHashFromParts(uint64* Parts)
-{
-	FHash128 Result;
-	memcpy(&Result, Parts, sizeof(Result));
-	return Result;
-}
-
-struct FPackWriteContext
-{
-	FPackWriteContext(const FPath& InOutputRoot) : OutputRoot(InOutputRoot) { Reset(); }
-	~FPackWriteContext() { FinishPack(); }
-
-	void AddBlock(const FGenericBlock& Block, FHash128 CompressedHash, FBufferView CompressedData)
-	{
-		std::lock_guard<std::mutex> LockGuard(Mutex);
-
-		UNSYNC_ASSERT(CompressedData.Size <= GMaxPackFileSize);
-
-		if (PackBuffer.Size() + CompressedData.Size > GMaxPackFileSize)
-		{
-			FinishPack();
-		}
-
-		FPackIndexEntry IndexEntry;
-		IndexEntry.BlockHash	  = Block.HashStrong.ToHash128();
-		IndexEntry.CompressedHash = CompressedHash;
-		IndexEntry.Offset		  = CheckedNarrow(PackBuffer.Size());
-		IndexEntry.CompressedSize = CheckedNarrow(CompressedData.Size);
-
-		IndexEntries.push_back(IndexEntry);
-		PackBuffer.Append(CompressedData);
-
-		UNSYNC_ASSERT(PackBuffer.Size() == IndexEntry.Offset + IndexEntry.CompressedSize);
-
-		AddHash(IndexFileHashSum, IndexEntry.BlockHash);
-	}
-
-	void FinishPack()
-	{
-		if (IndexEntries.empty())
-		{
-			return;
-		}
-
-		FHash128	BlockHash128 = MakeHashFromParts(IndexFileHashSum);
-		std::string OutputId	 = HashToHexString(BlockHash128);
-
-		FPath FinalPackFilename	 = OutputRoot / (OutputId + ".unsync_pack");
-		FPath FinalIndexFilename = OutputRoot / (OutputId + ".unsync_index");
-
-		UNSYNC_LOG(L"Saving new pack: %hs", OutputId.c_str());
-
-		if (!WriteBufferToFile(FinalPackFilename, PackBuffer, EFileMode::CreateWriteOnly))
-		{
-			UNSYNC_FATAL(L"Failed to write pack file '%ls'", FinalPackFilename.wstring().c_str());
-		}
-
-		const uint8* IndexData	   = reinterpret_cast<const uint8*>(IndexEntries.data());
-		uint64		 IndexDataSize = sizeof(IndexEntries[0]) * IndexEntries.size();
-		if (!WriteBufferToFile(FinalIndexFilename, IndexData, IndexDataSize, EFileMode::CreateWriteOnly))
-		{
-			UNSYNC_FATAL(L"Failed to write index file '%ls'", FinalIndexFilename.wstring().c_str());
-		}
-
-		Reset();
-	}
-
-private:
-	void Reset()
-	{
-		PackBuffer.Reserve(GMaxPackFileSize);
-		PackBuffer.Clear();
-		IndexEntries.clear();
-
-		IndexFileHashSum[0] = 0;
-		IndexFileHashSum[1] = 0;
-	}
-
-	std::mutex Mutex;
-
-	// Independent sums of low and high 32 bits of all seen block hashes.
-	// Used to generate a stable hash while allowing out-of-order block processing.
-	uint64 IndexFileHashSum[2] = {};
-
-	FBuffer						 PackBuffer;
-	std::vector<FPackIndexEntry> IndexEntries;
-
-	FPath OutputRoot;
-};
-
-static bool
-EnsureDirectoryExists(const FPath& Path)
-{
-	return (PathExists(Path) && IsDirectory(Path)) || CreateDirectories(Path);
-}
-
 static int32
 RunSubprocess(const char* Command, const FPath& WorkingDirectory, std::string& StdOutBuffer)
 {
@@ -536,9 +420,12 @@ CmdPack(const FCmdPackOptions& Options)
 	std::mutex Mutex;
 
 	FOnBlockGenerated OnBlockGenerated =
-		[&Mutex, &PackWriter, &SeenBlockHashSet, &ProcessedRawBytes, &CompressedBytes, &LogConfig](const FGenericBlock& Block,
-																								   FBufferView			Data)
+		[&Mutex, &PackWriter, &SeenBlockHashSet, &ProcessedRawBytes, &CompressedBytes, &LogConfig](const FGenericBlock&	   Block,
+																								   const FBlockSourceInfo& SourceInfo,
+																								   FBufferView			   Data)
 	{
+		UNSYNC_UNUSED(SourceInfo);
+
 		FThreadLogConfig::FScope LogConfigScope(LogConfig);
 
 		{
@@ -549,22 +436,12 @@ CmdPack(const FCmdPackOptions& Options)
 			}
 		}
 
-		const uint64 MaxCompressedSize	  = GetMaxCompressedSize(Block.Size);
-		FIOBuffer	 CompressedData		  = FIOBuffer::Alloc(MaxCompressedSize, L"PackBlock");
-		uint64		 ActualCompressedSize = CompressInto(Data, CompressedData.GetMutBufferView(), 9);
-
-		if (!ActualCompressedSize)
-		{
-			UNSYNC_FATAL(L"Failed to compress file block");
-		}
-		CompressedData.SetDataRange(0, ActualCompressedSize);
+		FPackWriteContext::FCompressedBlock CompressedBlock = FPackWriteContext::CompressBlock(Data);
 
 		ProcessedRawBytes += Block.Size;
-		CompressedBytes += ActualCompressedSize;
-
-		FHash128 CompressedHash = HashBlake3Bytes<FHash128>(CompressedData.GetData(), ActualCompressedSize);
+		CompressedBytes += CompressedBlock.Data.GetSize();
 
-		PackWriter.AddBlock(Block, CompressedHash, CompressedData.GetBufferView());
+		PackWriter.CompressAndAddBlock(Block, Data);
 	};
 
 	FComputeBlocksParams BlockParams;
@@ -637,13 +514,13 @@ CmdPack(const FCmdPackOptions& Options)
 
 		std::vector<FGenericBlock> ManifestBlocks;
 		FOnBlockGenerated		   OnManifestBlockGenerated =
-			[&OnBlockGenerated, &Mutex, &ManifestBlocks](const FGenericBlock& Block, FBufferView Data)
+			[&OnBlockGenerated, &Mutex, &ManifestBlocks](const FGenericBlock& Block, const FBlockSourceInfo& SourceInfo, FBufferView Data)
 		{
 			{
 				std::lock_guard<std::mutex> LockGuard(Mutex);
 				ManifestBlocks.push_back(Block);
 			}
-			OnBlockGenerated(Block, Data);
+			OnBlockGenerated(Block, SourceInfo, Data);
 		};
 
 		ManifestUniqueBytes -= ProcessedRawBytes.load();
@@ -740,7 +617,7 @@ BuildTargetFromPack(FIOWriter& Output, const FPackDatabase& PackDb, TArrayView<F
 			{
 				return PackIndex < Other.PackIndex;
 			}
-			return IndexEntry->Offset < Other.IndexEntry->Offset;
+			return IndexEntry->PackBlockOffset < Other.IndexEntry->PackBlockOffset;
 		}
 	};
 
@@ -779,7 +656,7 @@ BuildTargetFromPack(FIOWriter& Output, const FPackDatabase& PackDb, TArrayView<F
 	{
 		const FScheduleItem& Item = Schedule[ScheduleIndex];
 
-		UNSYNC_ASSERT(ReadSize == Item.IndexEntry->CompressedSize);
+		UNSYNC_ASSERT(ReadSize == Item.IndexEntry->PackBlockSize);
 
 		FHash128 CompressedHash = HashBlake3Bytes<FHash128>(Buffer.GetData(), Buffer.GetSize());
 		UNSYNC_ASSERT(CompressedHash == Item.IndexEntry->CompressedHash);
@@ -817,9 +694,9 @@ BuildTargetFromPack(FIOWriter& Output, const FPackDatabase& PackDb, TArrayView<F
 			CurrentPackId = Item.PackIndex;
 		}
 
-		UNSYNC_ASSERT(Item.IndexEntry->Offset < PackFile->GetSize());
-		UNSYNC_ASSERT(Item.IndexEntry->Offset + Item.IndexEntry->CompressedSize <= PackFile->GetSize());
-		PackFile->ReadAsync(Item.IndexEntry->Offset, Item.IndexEntry->CompressedSize, ScheduleIndex, ReadCallback);
+		UNSYNC_ASSERT(Item.IndexEntry->PackBlockOffset < PackFile->GetSize());
+		UNSYNC_ASSERT(Item.IndexEntry->PackBlockOffset + Item.IndexEntry->PackBlockSize <= PackFile->GetSize());
+		PackFile->ReadAsync(Item.IndexEntry->PackBlockOffset, Item.IndexEntry->PackBlockSize, ScheduleIndex, ReadCallback);
 	}
 
 	if (PackFile)
diff --git a/Engine/Source/Programs/Unsync/Private/UnsyncCore.cpp b/Engine/Source/Programs/Unsync/Private/UnsyncCore.cpp
--- a/Engine/Source/Programs/Unsync/Private/UnsyncCore.cpp
+++ b/Engine/Source/Programs/Unsync/Private/UnsyncCore.cpp
@@ -30,7 +30,7 @@ UNSYNC_THIRD_PARTY_INCLUDES_START
 #include <md5-sse2.h>
 UNSYNC_THIRD_PARTY_INCLUDES_END
 
-#define UNSYNC_VERSION_STR "1.0.72"
+#define UNSYNC_VERSION_STR "1.0.78"
 
 namespace unsync {
 
@@ -140,6 +140,9 @@ ComputeBlocksVariableT(FIOReader& Reader, const FComputeBlocksParams& Params)
 	const uint32 BlocksPerMacroBlock	 = CheckedNarrow(DivUp(TargetMacroBlockSize - MinimumMacroBlockSize, BlockSize));
 	const uint32 MacroBlockHashThreshold = BlocksPerMacroBlock ? (0xFFFFFFFF / BlocksPerMacroBlock) : 0;
 
+	FBlockSourceInfo SourceInfo;
+	SourceInfo.TotalSize = InputSize;
+
 	struct FTask
 	{
 		uint64			   Offset = 0;
@@ -178,6 +181,7 @@ ComputeBlocksVariableT(FIOReader& Reader, const FComputeBlocksParams& Params)
 						 &IoSemaphore,
 						 &BufferPool,
 						 &Params,
+						 &SourceInfo,
 						 MinimumBlockSize,
 						 MaximumBlockSize,
 						 ScanTaskBuffer,
@@ -211,7 +215,8 @@ ComputeBlocksVariableT(FIOReader& Reader, const FComputeBlocksParams& Params)
 						   &Params,
 						   TargetMacroBlockSize,
 						   &MacroBlockHasher,
-						   &CurrentMacroBlock](const uint8* WindowBegin, const uint8* WindowEnd, uint32 WindowHash)
+						   &CurrentMacroBlock,
+						   &SourceInfo](const uint8* WindowBegin, const uint8* WindowEnd, uint32 WindowHash)
 							  UNSYNC_ATTRIB_FORCEINLINE {
 								  // WARNING: Changing this invalidates some of the previously cached blocks.
 								  // TODO: compute based on target average block size
@@ -250,7 +255,7 @@ ComputeBlocksVariableT(FIOReader& Reader, const FComputeBlocksParams& Params)
 												  FBufferView BlockView;
 												  BlockView.Data = (LastBlockEnd + Block.Size) - CurrentMacroBlock.Size;
 												  BlockView.Size = CurrentMacroBlock.Size;
-												  Params.OnMacroBlockGenerated(CurrentMacroBlock, BlockView);
+												  Params.OnMacroBlockGenerated(CurrentMacroBlock, SourceInfo, BlockView);
 											  }
 
 											  // Reset macro block state
@@ -265,7 +270,7 @@ ComputeBlocksVariableT(FIOReader& Reader, const FComputeBlocksParams& Params)
 										  FBufferView BlockView;
 										  BlockView.Data = LastBlockEnd;
 										  BlockView.Size = Block.Size;
-										  Params.OnBlockGenerated(Block, BlockView);
+										  Params.OnBlockGenerated(Block, SourceInfo, BlockView);
 									  }
 
 									  if (!Task.Blocks.empty())
diff --git a/Engine/Source/Programs/Unsync/Private/UnsyncCore.h b/Engine/Source/Programs/Unsync/Private/UnsyncCore.h
--- a/Engine/Source/Programs/Unsync/Private/UnsyncCore.h
+++ b/Engine/Source/Programs/Unsync/Private/UnsyncCore.h
@@ -30,6 +30,7 @@ struct FRemoteDesc;
 struct FIOReader;
 struct FIOWriter;
 struct FIOReaderWriter;
+struct FPackWriteContext;
 
 struct FIdentityHash32
 {
@@ -126,7 +127,13 @@ FNeedListSize ComputeNeedListSize(const FNeedList& NeedList);
 
 const std::string& GetVersionString();
 
-using FOnBlockGenerated = std::function<void(const FGenericBlock& Block, FBufferView Data)>;
+struct FBlockSourceInfo
+{
+	// FPath  FilePath; // TODO: pass this in to allow filtering by name, etc.
+	uint64 TotalSize = 0;  // Total size of the input file/buffer from which the block was generated
+};
+
+using FOnBlockGenerated = std::function<void(const FGenericBlock& Block, const FBlockSourceInfo& Source, FBufferView Data)>;
 
 struct FComputeBlocksParams
 {
diff --git a/Engine/Source/Programs/Unsync/Private/UnsyncFile.cpp b/Engine/Source/Programs/Unsync/Private/UnsyncFile.cpp
--- a/Engine/Source/Programs/Unsync/Private/UnsyncFile.cpp
+++ b/Engine/Source/Programs/Unsync/Private/UnsyncFile.cpp
@@ -1202,6 +1202,12 @@ CreateDirectories(const FPath& Path)
 	return std::filesystem::create_directories(ExtendedPath);
 }
 
+bool
+EnsureDirectoryExists(const FPath& Path)
+{
+	return (PathExists(Path) && IsDirectory(Path)) || CreateDirectories(Path);
+}
+
 bool
 FileRename(const FPath& From, const FPath& To, std::error_code& OutErrorCode)
 {
diff --git a/Engine/Source/Programs/Unsync/Private/UnsyncFile.h b/Engine/Source/Programs/Unsync/Private/UnsyncFile.h
--- a/Engine/Source/Programs/Unsync/Private/UnsyncFile.h
+++ b/Engine/Source/Programs/Unsync/Private/UnsyncFile.h
@@ -467,6 +467,7 @@ bool			IsDirectory(const FPath& Path);
 bool			PathExists(const FPath& Path);
 bool			PathExists(const FPath& Path, std::error_code& OutErrorCode);
 bool			CreateDirectories(const FPath& Path);
+bool			EnsureDirectoryExists(const FPath& Path);
 bool			FileRename(const FPath& From, const FPath& To, std::error_code& OutErrorCode);
 bool			FileCopy(const FPath& From, const FPath& To, std::error_code& OutErrorCode);
 bool			FileCopyOverwrite(const FPath& From, const FPath& To, std::error_code& OutErrorCode);
diff --git a/Engine/Source/Programs/Unsync/Private/UnsyncMain.cpp b/Engine/Source/Programs/Unsync/Private/UnsyncMain.cpp
--- a/Engine/Source/Programs/Unsync/Private/UnsyncMain.cpp
+++ b/Engine/Source/Programs/Unsync/Private/UnsyncMain.cpp
@@ -105,6 +105,8 @@ InnerMain(int Argc, char** Argv)
 	bool					 bNoSocketTimeout	 = false;
 	bool					 bNoOutputFiles  	 = false;
 	bool					 bNoOutputRevisions  = false;
+	bool					 bPackOnlySmallFiles = false;
+	bool					 bPackFiles		 = false;
 	int32					 CompressionLevel	 = 3;
 	uint32					 DiffBlockSize		 = uint32(4_KB);
 	uint32					 HashOrSyncBlockSize = uint32(64_KB);
@@ -158,6 +160,16 @@ InnerMain(int Argc, char** Argv)
 		"--update",
 		bIncrementalMode,
 		"Create a directory manifest incrementally, by updating an existing manifest if one exists (only process changed files)");
+	SubHash->add_flag(
+		"--pack-small-files",
+		bPackOnlySmallFiles,
+		"Small files will be copied to compressed pack files during manifest generation and stored next to the manifest. "
+		"This can make syncing more efficient by reducing the number of remote file handles that must be opened. Implies --pack. "
+		"Files less than 4MB in size are considered small.");
+	SubHash->add_flag(
+		"--pack",
+		bPackFiles,
+		"Input files will be copied to compressed pack files during manifest generation and stored next to the manifest file.");
 	SubCommands.push_back(SubHash);
 
 	// Configure pack
@@ -397,6 +409,11 @@ InnerMain(int Argc, char** Argv)
 		return ReturnCode;
 	}
 
+	if (bPackOnlySmallFiles)
+	{
+		bPackFiles = true;
+	}
+
 	if (Cli.get_subcommands().size() == 0)
 	{
 		wprintf(L"%hs", Cli.help().c_str());
@@ -846,12 +863,18 @@ InnerMain(int Argc, char** Argv)
 	{
 		FCmdHashOptions HashOptions;
 
-		HashOptions.Input		 = InputFilename;
-		HashOptions.Output		 = OutputFilename;
-		HashOptions.BlockSize	 = HashOrSyncBlockSize;
-		HashOptions.Algorithm	 = Algorithm;
-		HashOptions.bForce		 = bForceOperation;
-		HashOptions.bIncremental = bIncrementalMode;
+		HashOptions.Input			= InputFilename;
+		HashOptions.Output			= OutputFilename;
+		HashOptions.BlockSize		= HashOrSyncBlockSize;
+		HashOptions.Algorithm		= Algorithm;
+		HashOptions.bForce			= bForceOperation;
+		HashOptions.bIncremental	= bIncrementalMode;
+		HashOptions.bPackFiles		= bPackFiles;
+
+		if (bPackOnlySmallFiles)
+		{
+			HashOptions.MaxFileSizeToPack = 4_MB;
+		}
 
 		return CmdHash(HashOptions);
 	}
diff --git a/Engine/Source/Programs/Unsync/Private/UnsyncManifest.h b/Engine/Source/Programs/Unsync/Private/UnsyncManifest.h
--- a/Engine/Source/Programs/Unsync/Private/UnsyncManifest.h
+++ b/Engine/Source/Programs/Unsync/Private/UnsyncManifest.h
@@ -71,6 +71,8 @@ struct FDirectoryManifest
 
 	bool bHasFileRevisionControl = false;
 
+	std::vector<FHash128> PackReferences;
+
 	bool IsValid() const { return Version != EVersions::Invalid; }
 };
 
diff --git a/Engine/Source/Programs/Unsync/Private/UnsyncPack.cpp b/Engine/Source/Programs/Unsync/Private/UnsyncPack.cpp
new file mode 100644
--- /dev/null
+++ b/Engine/Source/Programs/Unsync/Private/UnsyncPack.cpp
@@ -0,0 +1,186 @@
+// Copyright Epic Games, Inc. All Rights Reserved.
+
+#include "UnsyncPack.h"
+#include "UnsyncFile.h"
+#include "UnsyncCompression.h"
+#include "UnsyncHashTable.h"
+#include "UnsyncCore.h"
+
+namespace unsync {
+
+inline void
+AddHash(uint64* Accumulator, const FHash128& Hash)
+{
+	uint64 BlockHashParts[2];
+	memcpy(BlockHashParts, &Hash, sizeof(FHash128));
+	Accumulator[0] += BlockHashParts[0];
+	Accumulator[1] += BlockHashParts[1];
+}
+
+inline FHash128
+MakeHashFromParts(uint64* Parts)
+{
+	FHash128 Result;
+	memcpy(&Result, Parts, sizeof(Result));
+	return Result;
+}
+
+FPackWriteContext::FPackWriteContext(const FPath& InOutputRoot) : OutputRoot(InOutputRoot)
+{
+	InternalReset();
+}
+
+FPackWriteContext::~FPackWriteContext()
+{
+	InternalFinishPack();
+}
+
+void
+FPackWriteContext::AddCompressedBlock(const FGenericBlock& Block, FHash128 CompressedHash, FBufferView CompressedData)
+{
+	std::lock_guard<std::mutex> LockGuard(Mutex);
+
+	UNSYNC_ASSERT(CompressedData.Size <= GMaxPackFileSize);
+
+	if (PackBuffer.Size() + CompressedData.Size > GMaxPackFileSize)
+	{
+		InternalFinishPack();
+	}
+
+	FPackIndexEntry IndexEntry;
+	IndexEntry.BlockHash	  = Block.HashStrong.ToHash128();
+	IndexEntry.CompressedHash = CompressedHash;
+	IndexEntry.PackBlockOffset = CheckedNarrow(PackBuffer.Size());
+	IndexEntry.PackBlockSize   = CheckedNarrow(CompressedData.Size);
+
+	IndexEntries.push_back(IndexEntry);
+	PackBuffer.Append(CompressedData);
+
+	UNSYNC_ASSERT(PackBuffer.Size() == IndexEntry.PackBlockOffset + IndexEntry.PackBlockSize);
+
+	AddHash(IndexFileHashSum, IndexEntry.BlockHash);
+
+	ProcessedRawBytes += Block.Size;
+	ProcessedCompressedBytes = CompressedData.Size;
+}
+
+void
+FPackWriteContext::FinishPack()
+{
+	std::lock_guard<std::mutex> LockGuard(Mutex);
+	InternalFinishPack();
+}
+
+void
+FPackWriteContext::InternalFinishPack()
+{
+	// Assumes the function is called from a thread-safe context
+
+	if (IndexEntries.empty())
+	{
+		return;
+	}
+
+	FHash128	BlockHash128 = MakeHashFromParts(IndexFileHashSum);
+	std::string OutputId	 = HashToHexString(BlockHash128);
+
+	FPath FinalPackFilename	 = OutputRoot / (OutputId + ".unsync_pack");
+	FPath FinalIndexFilename = OutputRoot / (OutputId + ".unsync_index");
+
+	UNSYNC_LOG(L"Saving new pack: %hs", OutputId.c_str());
+
+	if (!GDryRun)
+	{
+		// TODO: offload final disk write to a background thread
+
+		if (!WriteBufferToFile(FinalPackFilename, PackBuffer, EFileMode::CreateWriteOnly))
+		{
+			UNSYNC_FATAL(L"Failed to write pack file '%ls'", FinalPackFilename.wstring().c_str());
+		}
+
+		const uint8* IndexData	   = reinterpret_cast<const uint8*>(IndexEntries.data());
+		uint64		 IndexDataSize = sizeof(IndexEntries[0]) * IndexEntries.size();
+
+		FPackIndexHeader IndexHeader;
+		const uint64	 IndexTotalSize = sizeof(IndexHeader) + IndexDataSize;
+
+		FNativeFile IndexFile(FinalIndexFilename, EFileMode::CreateWriteOnly, IndexTotalSize);
+		if (!IndexFile.IsValid())
+		{
+			UNSYNC_FATAL(L"Failed to write pack index file '%ls'", FinalIndexFilename.wstring().c_str());
+		}
+
+		uint64 WroteBytes = 0;
+		
+		WroteBytes += IndexFile.Write(&IndexHeader, WroteBytes, sizeof(IndexHeader));
+		WroteBytes += IndexFile.Write(IndexData, WroteBytes, IndexDataSize);
+
+		if (WroteBytes != IndexTotalSize)
+		{
+			UNSYNC_FATAL(L"Failed to write pack index file '%ls'", FinalIndexFilename.wstring().c_str());
+		}
+	}
+
+	GeneratedPackIds.push_back(BlockHash128);
+
+	InternalReset();
+}
+
+void
+FPackWriteContext::InternalReset()
+{
+	PackBuffer.Reserve(GMaxPackFileSize);
+	PackBuffer.Clear();
+	IndexEntries.clear();
+
+	IndexFileHashSum[0] = 0;
+	IndexFileHashSum[1] = 0;
+}
+
+FPackWriteContext::FCompressedBlock
+FPackWriteContext::CompressBlock(FBufferView RawData)
+{
+	FPackWriteContext::FCompressedBlock Result;
+
+	const uint64 MaxCompressedSize = GetMaxCompressedSize(RawData.Size);
+	Result.Data					   = FIOBuffer::Alloc(MaxCompressedSize, L"PackBlock");
+	uint64 ActualCompressedSize	   = CompressInto(RawData, Result.Data.GetMutBufferView(), 9);
+
+	if (!ActualCompressedSize)
+	{
+		UNSYNC_FATAL(L"Failed to compress file block");
+	}
+	Result.Data.SetDataRange(0, ActualCompressedSize);
+
+	Result.Hash = HashBlake3Bytes<FHash128>(Result.Data.GetData(), ActualCompressedSize);
+
+	return Result;
+}
+
+void FPackWriteContext::CompressAndAddBlock(const FGenericBlock& Block, FBufferView RawData)
+{
+	FCompressedBlock Compressed = CompressBlock(RawData);
+	AddCompressedBlock(Block, Compressed.Hash, Compressed.Data.GetBufferView());
+}
+
+void
+FPackWriteContext::GetUniqueGeneratedPackIds(std::vector<FHash128>& Output)
+{
+	THashSet<FHash128> KnownHashes;
+	for (const FHash128& Hash : Output)
+	{
+		KnownHashes.insert(Hash);
+	}
+
+	std::lock_guard<std::mutex> LockGuard(Mutex);
+
+	for (const FHash128& Hash : GeneratedPackIds)
+	{
+		if (KnownHashes.insert(Hash).second)
+		{
+			Output.push_back(Hash);
+		}
+	}
+}
+
+}  // namespace unsync
diff --git a/Engine/Source/Programs/Unsync/Private/UnsyncPack.h b/Engine/Source/Programs/Unsync/Private/UnsyncPack.h
new file mode 100644
--- /dev/null
+++ b/Engine/Source/Programs/Unsync/Private/UnsyncPack.h
@@ -0,0 +1,58 @@
+// Copyright Epic Games, Inc. All Rights Reserved.
+
+#pragma once
+
+#include "UnsyncCommon.h"
+#include "UnsyncHash.h"
+#include "UnsyncProtocol.h"
+#include "UnsyncFile.h"
+
+#include <mutex>
+#include <vector>
+
+namespace unsync {
+
+static constexpr uint64 GMaxPackFileSize = 1_GB;
+
+struct FPackWriteContext
+{
+	FPackWriteContext(const FPath& InOutputRoot);
+	~FPackWriteContext();
+
+	void AddCompressedBlock(const FGenericBlock& Block, FHash128 CompressedHash, FBufferView CompressedData); // thread-safe
+	void CompressAndAddBlock(const FGenericBlock& Block, FBufferView RawData);
+
+	void FinishPack();
+
+	// Common utility function to compress a block and generate compressed hash
+	struct FCompressedBlock
+	{
+		FIOBuffer Data; // zstd-compressed block
+		FHash128  Hash; // block hash after compression
+	};
+	static FCompressedBlock CompressBlock(FBufferView RawData);
+
+	void GetUniqueGeneratedPackIds(std::vector<FHash128>& Output);
+
+private:
+	void InternalReset();
+	void InternalFinishPack();
+
+	std::mutex Mutex;
+
+	// Independent sums of low and high 32 bits of all seen block hashes.
+	// Used to generate a stable hash while allowing out-of-order block processing.
+	uint64 IndexFileHashSum[2] = {};
+
+	FBuffer						 PackBuffer;
+	std::vector<FPackIndexEntry> IndexEntries;
+
+	uint64 ProcessedRawBytes = 0;
+	uint64 ProcessedCompressedBytes = 0;
+
+	FPath OutputRoot;
+
+	std::vector<FHash128> GeneratedPackIds;
+};
+
+}  // namespace unsync
diff --git a/Engine/Source/Programs/Unsync/Private/UnsyncProtocol.h b/Engine/Source/Programs/Unsync/Private/UnsyncProtocol.h
--- a/Engine/Source/Programs/Unsync/Private/UnsyncProtocol.h
+++ b/Engine/Source/Programs/Unsync/Private/UnsyncProtocol.h
@@ -75,7 +75,8 @@ static constexpr uint64 SERIALIZED_SECTION_ID_TERMINATOR			= 0;
 static constexpr uint64 SERIALIZED_SECTION_ID_METADATA_STRING		= 0xC6BD6CDCEEF79533ull;
 static constexpr uint64 SERIALIZED_SECTION_ID_MACRO_BLOCK			= 0x8390AEBB745E08BCull;
 static constexpr uint64 SERIALIZED_SECTION_ID_FILE_READ_ONLY_MASK	= 0x851F32ED3615F0ADull;
-static constexpr uint64 SERIALIZED_SECTION_ID_FILE_REVISION_CONTROL = 0X2C1C72E6B78B1B50ull;
+static constexpr uint64 SERIALIZED_SECTION_ID_FILE_REVISION_CONTROL = 0x2C1C72E6B78B1B50ull;
+static constexpr uint64 SERIALIZED_SECTION_ID_PACK_REFERENCE		= 0x634EA57F1E48DFBDull;
 
 struct FSerializedSectionHeader
 {
@@ -112,6 +113,12 @@ struct FFileRevisionControlSection
 	static constexpr uint64 VERSION = 1;
 };
 
+struct FPackReferenceSection
+{
+	static constexpr uint64 MAGIC	= SERIALIZED_SECTION_ID_PACK_REFERENCE;
+	static constexpr uint64 VERSION = 1;
+};
+
 struct FBlockFileHeader
 {
 	static constexpr uint64 MAGIC	= 0x1DCB86A5BDBA27CFull;
@@ -202,6 +209,29 @@ struct FPatchHeader
 	EStrongHashAlgorithmID StrongHashAlgorithmId	 = EStrongHashAlgorithmID::Blake3_128;
 };
 
+struct FPackIndexEntry
+{
+	// Decompressed block hash
+	FHash128 BlockHash		 = {};
+	// Compressed block hash (may be equal to BlockHash to signal uncompressed block)
+	FHash128 CompressedHash	 = {};
+
+	// Offset and size within pack file
+	uint32	 PackBlockOffset = 0;
+	uint32	 PackBlockSize	 = 0;
+};
+static_assert(sizeof(FPackIndexEntry) == 40);
+
+struct FPackIndexHeader
+{
+	static constexpr uint64 MAGIC	= 0xEEC735E03053CC3Full;
+	static constexpr uint64 VERSION = 1;
+
+	uint64 Magic   = MAGIC;
+	uint64 Version = VERSION;
+};
+static_assert(sizeof(FPackIndexHeader) == 16);
+
 // Protocol V2: support for up to 256bit hashes
 
 struct FBlockRequest256
diff --git a/Engine/Source/Programs/Unsync/Private/UnsyncScheduler.h b/Engine/Source/Programs/Unsync/Private/UnsyncScheduler.h
--- a/Engine/Source/Programs/Unsync/Private/UnsyncScheduler.h
+++ b/Engine/Source/Programs/Unsync/Private/UnsyncScheduler.h
@@ -33,6 +33,8 @@ class FScheduler
 public:
 	UNSYNC_DISALLOW_COPY_ASSIGN(FScheduler)
 
+	using FTaskFunction = FThreadPool::FTaskFunction;
+
 	static constexpr uint32 MAX_NETWORK_TASKS = 8;
 
 	FScheduler(uint32 InNumWorkerThreads);
@@ -56,6 +58,12 @@ public:
 		return NumExecuted != 0;
 	}
 
+	template<typename TaskFunction>
+	void PushTask(TaskFunction&& Fun, bool bAllowImmediateExecution = true)
+	{
+		ThreadPool.PushTask(std::forward<TaskFunction>(Fun), bAllowImmediateExecution);
+	}
+
 private:
 	FThreadPool ThreadPool;
 };
diff --git a/Engine/Source/Programs/Unsync/Private/UnsyncSerialization.cpp b/Engine/Source/Programs/Unsync/Private/UnsyncSerialization.cpp
--- a/Engine/Source/Programs/Unsync/Private/UnsyncSerialization.cpp
+++ b/Engine/Source/Programs/Unsync/Private/UnsyncSerialization.cpp
@@ -376,7 +376,6 @@ SaveFileRevisionControl(const FDirectoryManifest& Manifest)
 	return Result;
 }
 
-
 std::vector<std::string>
 LoadFileRevisionControl(FIOReaderStream& Reader, FSerializedSectionHeader Header)
 {
@@ -395,6 +394,40 @@ LoadFileRevisionControl(FIOReaderStream& Reader, FSerializedSectionHeader Header
 	return Result;
 }
 
+static FBuffer
+SavePackReferences(const FDirectoryManifest& Manifest)
+{
+	FBuffer			 Result;
+	FVectorStreamOut Writer(Result);
+
+	const uint64 NumEntries = Manifest.PackReferences.size();
+	Writer.WriteT(NumEntries);
+	for (const FHash128& Hash : Manifest.PackReferences)
+	{
+		Writer.WriteT(Hash);
+	}
+
+	return Result;
+}
+
+static std::vector<FHash128>
+LoadPackReferences(FIOReaderStream& Reader, FSerializedSectionHeader Header)
+{
+	std::vector<FHash128> Result;
+
+	uint64 NumEntries = 0;
+	Serialize(Reader, NumEntries);
+
+	Result.resize(NumEntries);
+
+	for (uint64 i = 0; i < NumEntries; ++i)
+	{
+		Serialize(Reader, Result[i]);
+	}
+
+	return Result;
+}
+
 static void
 ConvertSerializedPathSeparatorsToNative(std::string& PathUtf8)
 {
@@ -508,6 +541,11 @@ LoadDirectoryManifest(FDirectoryManifest& OutManifest, const FPath& Root, FIORea
 							OutManifest.bHasFileRevisionControl = true;
 							break;
 						}
+					case SERIALIZED_SECTION_ID_PACK_REFERENCE:
+						{
+							OutManifest.PackReferences = LoadPackReferences(Stream, SectionHeader);
+							break;
+						}
 					case SERIALIZED_SECTION_ID_TERMINATOR:
 						bDoneLoadingOptionalSections = true;
 						break;
@@ -753,6 +791,12 @@ SaveDirectoryManifest(const FDirectoryManifest& Manifest, FVectorStreamOut& Stre
 		WriteSection<FFileRevisionControlSection>(Stream, SectionBuffer.View());
 	}
 
+	if (!Manifest.PackReferences.empty())
+	{
+		FBuffer SectionBuffer = SavePackReferences(Manifest);
+		WriteSection<FPackReferenceSection>(Stream, SectionBuffer.View());
+	}
+
 	// End with the terminator section (default-constructed);
 	FSerializedSectionHeader TerminatorSection;
 	Serialize(Stream, TerminatorSection);
diff --git a/Engine/Source/Programs/Unsync/Private/UnsyncThread.cpp b/Engine/Source/Programs/Unsync/Private/UnsyncThread.cpp
--- a/Engine/Source/Programs/Unsync/Private/UnsyncThread.cpp
+++ b/Engine/Source/Programs/Unsync/Private/UnsyncThread.cpp
@@ -29,7 +29,7 @@ FThreadPool::StartWorkers(uint32 NumWorkers)
 FThreadPool::~FThreadPool()
 {
 	bShutdown = true;
-	WorkerWakeupCondition.notify_all();
+	WakeCondition.notify_all();
 
 	for (std::thread& Thread : Threads)
 	{
@@ -45,7 +45,7 @@ FThreadPool::PopTask(bool bWaitForSignal)
 	if (bWaitForSignal)
 	{
 		auto WaitUntil = [this]() { return bShutdown || !Tasks.empty(); };
-		WorkerWakeupCondition.wait(LockScope, WaitUntil);
+		WakeCondition.wait(LockScope, WaitUntil);
 	}
 
 	FThreadPool::FTaskFunction Result;
@@ -60,9 +60,9 @@ FThreadPool::PopTask(bool bWaitForSignal)
 }
 
 void
-FThreadPool::PushTask(FTaskFunction&& Fun)
+FThreadPool::PushTask(FTaskFunction&& Fun, bool bAllowImmediateExecution)
 {
-	if (Threads.empty() || NumRunningTasks.load() == NumWorkerThreads())
+	if (Threads.empty() || (NumRunningTasks.load() == NumWorkerThreads() && bAllowImmediateExecution))
 	{
 		Fun();
 	}
@@ -70,7 +70,7 @@ FThreadPool::PushTask(FTaskFunction&& Fun)
 	{
 		std::unique_lock<std::mutex> LockScope(Mutex);
 		Tasks.push_back(std::forward<FTaskFunction>(Fun));
-		WorkerWakeupCondition.notify_one();
+		WakeCondition.notify_one();
 	}
 }
 
@@ -95,7 +95,6 @@ FThreadPool::DoWorkInternal(bool bWaitForSignal)
 	}
 }
 
-
 void
 SchedulerSleep(uint32 Milliseconds)
 {
diff --git a/Engine/Source/Programs/Unsync/Private/UnsyncThread.h b/Engine/Source/Programs/Unsync/Private/UnsyncThread.h
--- a/Engine/Source/Programs/Unsync/Private/UnsyncThread.h
+++ b/Engine/Source/Programs/Unsync/Private/UnsyncThread.h
@@ -10,7 +10,6 @@
 UNSYNC_THIRD_PARTY_INCLUDES_START
 #include <atomic>
 #include <condition_variable>
-#include <deque>
 #include <functional>
 #include <thread>
 #include <vector>
@@ -86,11 +85,11 @@ public:
 	// Does nothing if the number of already launched workers is lower than given value.
 	void StartWorkers(uint32 NumWorkers);
 
-	// Adds a task to the FIFO queue
-	void PushTask(FTaskFunction&& Fun);
+	// Adds a task to the task list
+	void PushTask(FTaskFunction&& Fun, bool bAllowImmediateExecution = true);
 
-	// Try to pop the next task from the queue and execute it on the current thread.
-	// Returns false if queue is empty, which may happen if worker threads have picked up the tasks already.
+	// Try to pop the next task from the list and execute it on the current thread.
+	// Returns false if list is empty, which may happen if worker threads have picked up the tasks already.
 	bool TryExecuteTask() { return DoWorkInternal(false); }
 
 	uint32 NumWorkerThreads() const { return uint32(Threads.size()); }
@@ -106,7 +105,7 @@ private:
 	std::vector<FTaskFunction> Tasks;
 
 	std::mutex				Mutex;
-	std::condition_variable WorkerWakeupCondition;
+	std::condition_variable WakeCondition;
 	std::atomic<bool>		bShutdown;
 	std::atomic<uint64>		NumRunningTasks;
 };
