unsync - Store uncompressed blocks by in packs by default

diff --git a/Engine/Source/Programs/Unsync/Private/UnsyncCmdHash.cpp b/Engine/Source/Programs/Unsync/Private/UnsyncCmdHash.cpp
--- a/Engine/Source/Programs/Unsync/Private/UnsyncCmdHash.cpp
+++ b/Engine/Source/Programs/Unsync/Private/UnsyncCmdHash.cpp
@@ -45,13 +45,14 @@ CmdHash(const FCmdHashOptions& Options)
 		// TODO: macro block generation is only implemented for variable chunk mode
 		ComputeBlocksParams.bNeedMacroBlocks = ComputeBlocksParams.Algorithm.ChunkingAlgorithmId == EChunkingAlgorithmID::VariableBlocks;
 
+		FPath PackOutputRoot = ManifestRoot / "pack";
+
 		std::unique_ptr<FPackWriteContext> PackWriter;
 		THashSet<FGenericHash>			   PackedBlocks;
 		std::mutex						   PackedBlocksMutex;
 
 		if (!GDryRun && Options.bPackFiles)
 		{
-			FPath PackOutputRoot = ManifestRoot / "pack";
 
 			if (!PathExists(PackOutputRoot))
 			{
@@ -85,7 +86,8 @@ CmdHash(const FCmdHashOptions& Options)
 					}
 				}
 
-				PackWriter->CompressAndAddBlock(Block, Data);
+				// TODO(pack): compressed packs will require special handling by the server
+				PackWriter->AddRawBlock(Block, Data);
 			};
 
 			ComputeBlocksParams.OnBlockGenerated = OnBlockGenerated;
diff --git a/Engine/Source/Programs/Unsync/Private/UnsyncFile.cpp b/Engine/Source/Programs/Unsync/Private/UnsyncFile.cpp
--- a/Engine/Source/Programs/Unsync/Private/UnsyncFile.cpp
+++ b/Engine/Source/Programs/Unsync/Private/UnsyncFile.cpp
@@ -1246,6 +1246,13 @@ std::filesystem::recursive_directory_iterator RecursiveDirectoryScan(const FPath
 	return std::filesystem::recursive_directory_iterator(ExtendedPath);
 }
 
+std::filesystem::directory_iterator
+DirectoryScan(const FPath& Path)
+{
+	FPath ExtendedPath = MakeExtendedAbsolutePath(Path);
+	return std::filesystem::directory_iterator(ExtendedPath);
+}
+
 FMemReader::FMemReader(const uint8* InData, uint64 InDataSize) : Data(InData), Size(InDataSize)
 {
 }
diff --git a/Engine/Source/Programs/Unsync/Private/UnsyncFile.h b/Engine/Source/Programs/Unsync/Private/UnsyncFile.h
--- a/Engine/Source/Programs/Unsync/Private/UnsyncFile.h
+++ b/Engine/Source/Programs/Unsync/Private/UnsyncFile.h
@@ -479,6 +479,7 @@ FPath			GetRelativePath(const FPath& Path, const FPath& Base);
 uint64 GetAvailableDiskSpace(const FPath& Path);
 
 std::filesystem::recursive_directory_iterator RecursiveDirectoryScan(const FPath& Path);
+std::filesystem::directory_iterator DirectoryScan(const FPath& Path);
 
 uint64 ToWindowsFileTime(const std::filesystem::file_time_type& T);
 std::filesystem::file_time_type FromWindowsFileTime(uint64 Ticks);
diff --git a/Engine/Source/Programs/Unsync/Private/UnsyncManifest.h b/Engine/Source/Programs/Unsync/Private/UnsyncManifest.h
--- a/Engine/Source/Programs/Unsync/Private/UnsyncManifest.h
+++ b/Engine/Source/Programs/Unsync/Private/UnsyncManifest.h
@@ -71,7 +71,7 @@ struct FDirectoryManifest
 
 	bool bHasFileRevisionControl = false;
 
-	std::vector<FHash128> PackReferences;
+	std::vector<FPackReference> PackReferences;
 
 	bool IsValid() const { return Version != EVersions::Invalid; }
 };
diff --git a/Engine/Source/Programs/Unsync/Private/UnsyncPack.cpp b/Engine/Source/Programs/Unsync/Private/UnsyncPack.cpp
--- a/Engine/Source/Programs/Unsync/Private/UnsyncPack.cpp
+++ b/Engine/Source/Programs/Unsync/Private/UnsyncPack.cpp
@@ -5,6 +5,7 @@
 #include "UnsyncCompression.h"
 #include "UnsyncHashTable.h"
 #include "UnsyncCore.h"
+#include "UnsyncManifest.h"
 
 namespace unsync {
 
@@ -35,14 +36,28 @@ FPackWriteContext::~FPackWriteContext()
 	InternalFinishPack();
 }
 
+void
+FPackWriteContext::AddRawBlock(const FGenericBlock& Block, FBufferView RawData)
+{
+	std::lock_guard<std::mutex> LockGuard(Mutex);
+	InternalAddBlock(Block, Block.HashStrong.ToHash128(), RawData);
+	++NumRawBlocks;
+}
+
 void
 FPackWriteContext::AddCompressedBlock(const FGenericBlock& Block, FHash128 CompressedHash, FBufferView CompressedData)
 {
 	std::lock_guard<std::mutex> LockGuard(Mutex);
+	InternalAddBlock(Block, CompressedHash, CompressedData);
+	++NumCompressedBlocks;
+}
 
-	UNSYNC_ASSERT(CompressedData.Size <= GMaxPackFileSize);
+void
+FPackWriteContext::InternalAddBlock(const FGenericBlock& Block, FHash128 CompressedHash, FBufferView Data)
+{
+	UNSYNC_ASSERT(Data.Size <= GMaxPackFileSize);
 
-	if (PackBuffer.Size() + CompressedData.Size > GMaxPackFileSize)
+	if (PackBuffer.Size() + Data.Size > GMaxPackFileSize)
 	{
 		InternalFinishPack();
 	}
@@ -51,17 +66,17 @@ FPackWriteContext::AddCompressedBlock(const FGenericBlock& Block, FHash128 Compr
 	IndexEntry.BlockHash	  = Block.HashStrong.ToHash128();
 	IndexEntry.CompressedHash = CompressedHash;
 	IndexEntry.PackBlockOffset = CheckedNarrow(PackBuffer.Size());
-	IndexEntry.PackBlockSize   = CheckedNarrow(CompressedData.Size);
+	IndexEntry.PackBlockSize   = CheckedNarrow(Data.Size);
 
 	IndexEntries.push_back(IndexEntry);
-	PackBuffer.Append(CompressedData);
+	PackBuffer.Append(Data);
 
 	UNSYNC_ASSERT(PackBuffer.Size() == IndexEntry.PackBlockOffset + IndexEntry.PackBlockSize);
 
 	AddHash(IndexFileHashSum, IndexEntry.BlockHash);
 
 	ProcessedRawBytes += Block.Size;
-	ProcessedCompressedBytes = CompressedData.Size;
+	ProcessedCompressedBytes = Data.Size;
 }
 
 void
@@ -87,7 +102,10 @@ FPackWriteContext::InternalFinishPack()
 	FPath FinalPackFilename	 = OutputRoot / (OutputId + ".unsync_pack");
 	FPath FinalIndexFilename = OutputRoot / (OutputId + ".unsync_index");
 
-	UNSYNC_LOG(L"Saving new pack: %hs", OutputId.c_str());
+	// Force non-indented log
+	FLogIndentScope IndentScope(0, true /*override*/);
+
+	UNSYNC_LOG(L"* Saving new pack: %hs", OutputId.c_str());
 
 	if (!GDryRun)
 	{
@@ -121,7 +139,23 @@ FPackWriteContext::InternalFinishPack()
 		}
 	}
 
-	GeneratedPackIds.push_back(BlockHash128);
+	FPackReference PackReference;
+	PackReference.Id = BlockHash128;
+
+	if (NumRawBlocks != 0)
+	{
+		PackReference.Flags = PackReference.Flags | EPackReferenceFlags::HasRawBlocks;
+	}
+
+	if (NumCompressedBlocks != 0)
+	{
+		PackReference.Flags = PackReference.Flags | EPackReferenceFlags::HasCompressedBlocks;
+	}
+
+	PackReference.NumTotalBlocks = NumCompressedBlocks + NumRawBlocks;
+	PackReference.NumUsedBlocks	 = PackReference.NumTotalBlocks; // By default assume all blocks in the pack are referenced
+
+	GeneratedPackIds.push_back(PackReference);
 
 	InternalReset();
 }
@@ -135,6 +169,9 @@ FPackWriteContext::InternalReset()
 
 	IndexFileHashSum[0] = 0;
 	IndexFileHashSum[1] = 0;
+
+	NumCompressedBlocks = 0;
+	NumRawBlocks		= 0;
 }
 
 FPackWriteContext::FCompressedBlock
@@ -164,21 +201,52 @@ void FPackWriteContext::CompressAndAddBlock(const FGenericBlock& Block, FBufferV
 }
 
 void
-FPackWriteContext::GetUniqueGeneratedPackIds(std::vector<FHash128>& Output)
+FPackWriteContext::GetUniqueGeneratedPackIds(std::vector<FPackReference>& Output) const
 {
-	THashSet<FHash128> KnownHashes;
-	for (const FHash128& Hash : Output)
+	std::lock_guard<std::mutex> LockGuard(Mutex);
+
+	MergePackReferences(Output, GeneratedPackIds);
+}
+
+void
+MergePackReferences(std::vector<FPackReference>& Destination, const std::vector<FPackReference>& Source)
+{
+	THashSet<FPackReference, FPackReference::Hasher> KnownHashes;
+
+	for (const FPackReference& Hash : Destination)
 	{
 		KnownHashes.insert(Hash);
 	}
 
-	std::lock_guard<std::mutex> LockGuard(Mutex);
-
-	for (const FHash128& Hash : GeneratedPackIds)
+	for (const FPackReference& Hash : Source)
 	{
 		if (KnownHashes.insert(Hash).second)
 		{
-			Output.push_back(Hash);
+			Destination.push_back(Hash);
+		}
+	}
+}
+
+void DeletePackAndIdexData(const FPath& PackRootDirectory)
+{
+	UNSYNC_VERBOSE(L"Deleting packs in '%ls'", PackRootDirectory.wstring().c_str());
+	for (const std::filesystem::directory_entry& Dir : DirectoryScan(PackRootDirectory))
+	{
+		if (!Dir.is_regular_file())
+		{
+			continue;
+		}
+
+		const FPath& FilePath	   = Dir.path();
+		const FPath	 FileExtension = FilePath.extension();
+		if (FileExtension == ".unsync_pack" || FileExtension == ".unsync_index")
+		{
+			std::error_code ErrorCode;
+			bool bDeleted = FileRemove(FilePath, ErrorCode);
+			if (!bDeleted)
+			{
+				UNSYNC_ERROR(L"Could not delete file '%ls'. Error code: %d.", FilePath.wstring().c_str(), ErrorCode.value());
+			}
 		}
 	}
 }
diff --git a/Engine/Source/Programs/Unsync/Private/UnsyncPack.h b/Engine/Source/Programs/Unsync/Private/UnsyncPack.h
--- a/Engine/Source/Programs/Unsync/Private/UnsyncPack.h
+++ b/Engine/Source/Programs/Unsync/Private/UnsyncPack.h
@@ -14,15 +14,18 @@ namespace unsync {
 
 static constexpr uint64 GMaxPackFileSize = 1_GB;
 
+struct FDirectoryManifest;
+
 struct FPackWriteContext
 {
 	FPackWriteContext(const FPath& InOutputRoot);
 	~FPackWriteContext();
 
+	void AddRawBlock(const FGenericBlock& Block, FBufferView RawData);  // thread-safe
 	void AddCompressedBlock(const FGenericBlock& Block, FHash128 CompressedHash, FBufferView CompressedData); // thread-safe
-	void CompressAndAddBlock(const FGenericBlock& Block, FBufferView RawData);
+	void CompressAndAddBlock(const FGenericBlock& Block, FBufferView RawData); // thread-safe
 
-	void FinishPack();
+	void FinishPack(); // thread-safe
 
 	// Common utility function to compress a block and generate compressed hash
 	struct FCompressedBlock
@@ -32,13 +35,15 @@ struct FPackWriteContext
 	};
 	static FCompressedBlock CompressBlock(FBufferView RawData);
 
-	void GetUniqueGeneratedPackIds(std::vector<FHash128>& Output);
+	void GetUniqueGeneratedPackIds(std::vector<FPackReference>& Output) const; // thread-safe
 
 private:
+
+	void InternalAddBlock(const FGenericBlock& Block, FHash128 CompressedHash, FBufferView Data);
 	void InternalReset();
 	void InternalFinishPack();
 
-	std::mutex Mutex;
+	mutable std::mutex Mutex;
 
 	// Independent sums of low and high 32 bits of all seen block hashes.
 	// Used to generate a stable hash while allowing out-of-order block processing.
@@ -50,9 +55,16 @@ private:
 	uint64 ProcessedRawBytes = 0;
 	uint64 ProcessedCompressedBytes = 0;
 
+	uint32 NumRawBlocks = 0;
+	uint32 NumCompressedBlocks = 0;
+
 	FPath OutputRoot;
 
-	std::vector<FHash128> GeneratedPackIds;
+	std::vector<FPackReference> GeneratedPackIds;
 };
 
+void MergePackReferences(std::vector<FPackReference>& Destination, const std::vector<FPackReference>& Source);
+
+void DeletePackAndIdexData(const FPath& PackRootDirectory);
+
 }  // namespace unsync
diff --git a/Engine/Source/Programs/Unsync/Private/UnsyncProtocol.h b/Engine/Source/Programs/Unsync/Private/UnsyncProtocol.h
--- a/Engine/Source/Programs/Unsync/Private/UnsyncProtocol.h
+++ b/Engine/Source/Programs/Unsync/Private/UnsyncProtocol.h
@@ -116,7 +116,7 @@ struct FFileRevisionControlSection
 struct FPackReferenceSection
 {
 	static constexpr uint64 MAGIC	= SERIALIZED_SECTION_ID_PACK_REFERENCE;
-	static constexpr uint64 VERSION = 1;
+	static constexpr uint64 VERSION = 3;
 };
 
 struct FBlockFileHeader
@@ -232,6 +232,38 @@ struct FPackIndexHeader
 };
 static_assert(sizeof(FPackIndexHeader) == 16);
 
+enum class EPackReferenceFlags : uint32 {
+	Default				= 0,
+	HasRawBlocks		= 1 << 0,
+	HasCompressedBlocks = 1 << 1,
+};
+UNSYNC_ENUM_CLASS_FLAGS(EPackReferenceFlags, uint32);
+
+// Basic information about a pack file referenced by a manifest
+struct FPackReference
+{
+	using EFlags = EPackReferenceFlags;
+
+	bool operator==(const FPackReference& Other) const { return 
+		Id == Other.Id && Flags == Other.Flags
+		&& NumUsedBlocks == Other.NumUsedBlocks
+		&& NumTotalBlocks == Other.NumTotalBlocks; }
+	struct Hasher
+	{
+		size_t operator()(const FPackReference& X) const
+		{
+			FHash128::Hasher H;
+			return H(X.Id);
+		}
+	};
+
+	FHash128 Id				= {};
+	EFlags	 Flags			= EFlags::Default;
+	uint32	 NumUsedBlocks	= 0;
+	uint32	 NumTotalBlocks = 0;
+};
+static_assert(sizeof(FPackReference) == 28);
+
 // Protocol V2: support for up to 256bit hashes
 
 struct FBlockRequest256
diff --git a/Engine/Source/Programs/Unsync/Private/UnsyncSerialization.cpp b/Engine/Source/Programs/Unsync/Private/UnsyncSerialization.cpp
--- a/Engine/Source/Programs/Unsync/Private/UnsyncSerialization.cpp
+++ b/Engine/Source/Programs/Unsync/Private/UnsyncSerialization.cpp
@@ -402,7 +402,7 @@ SavePackReferences(const FDirectoryManifest& Manifest)
 
 	const uint64 NumEntries = Manifest.PackReferences.size();
 	Writer.WriteT(NumEntries);
-	for (const FHash128& Hash : Manifest.PackReferences)
+	for (const FPackReference& Hash : Manifest.PackReferences)
 	{
 		Writer.WriteT(Hash);
 	}
@@ -410,10 +410,10 @@ SavePackReferences(const FDirectoryManifest& Manifest)
 	return Result;
 }
 
-static std::vector<FHash128>
+static std::vector<FPackReference>
 LoadPackReferences(FIOReaderStream& Reader, FSerializedSectionHeader Header)
 {
-	std::vector<FHash128> Result;
+	std::vector<FPackReference> Result;
 
 	uint64 NumEntries = 0;
 	Serialize(Reader, NumEntries);
