unsync - Move various utility code from UnsyncCore to separate modules

diff --git a/Engine/Source/Programs/Unsync/CMakeLists.txt b/Engine/Source/Programs/Unsync/CMakeLists.txt
--- a/Engine/Source/Programs/Unsync/CMakeLists.txt
+++ b/Engine/Source/Programs/Unsync/CMakeLists.txt
@@ -34,6 +34,8 @@ set(src
 	Private/UnsyncAuth.h
 	Private/UnsyncBuffer.cpp
 	Private/UnsyncBuffer.h
+	Private/UnsyncChunking.cpp
+	Private/UnsyncChunking.h
 	Private/UnsyncCmdDiff.cpp
 	Private/UnsyncCmdDiff.h
 	Private/UnsyncCmdHash.cpp
@@ -57,6 +59,8 @@ set(src
 	Private/UnsyncCompression.h
 	Private/UnsyncCore.cpp
 	Private/UnsyncCore.h
+	Private/UnsyncDiff.cpp
+	Private/UnsyncDiff.h
 	Private/UnsyncError.h
 	Private/UnsyncFile.cpp
 	Private/UnsyncFile.h
@@ -107,6 +111,8 @@ set(src
 	Private/UnsyncUtil.cpp
 	Private/UnsyncUtil.h
 	Private/UnsyncVarInt.h
+	Private/UnsyncVersion.cpp
+	Private/UnsyncVersion.h
 )
 
 # Add third party code
diff --git a/Engine/Source/Programs/Unsync/Private/UnsyncChunking.cpp b/Engine/Source/Programs/Unsync/Private/UnsyncChunking.cpp
new file mode 100644
--- /dev/null
+++ b/Engine/Source/Programs/Unsync/Private/UnsyncChunking.cpp
@@ -0,0 +1,515 @@
+// Copyright Epic Games, Inc. All Rights Reserved.
+
+#include "UnsyncChunking.h"
+#include "UnsyncFile.h"
+#include "UnsyncProtocol.h"
+#include "UnsyncScan.h"
+#include "UnsyncScheduler.h"
+
+UNSYNC_THIRD_PARTY_INCLUDES_START
+#include <blake3.h>
+#include <md5-sse2.h>
+UNSYNC_THIRD_PARTY_INCLUDES_END
+
+namespace unsync {
+
+inline uint32
+ComputeMinVariableBlockSize(uint32 BlockSize)
+{
+	return std::min(MAX_BLOCK_SIZE, std::max(BlockSize, 4096u) / 2);  // changing this invalidates cached blocks
+}
+
+inline uint32
+ComputeMaxVariableBlockSize(uint32 BlockSize)
+{
+	return std::min(MAX_BLOCK_SIZE, std::max(BlockSize, 4096u) * 4);  // changing this invalidates cached blocks
+}
+
+template<typename WeakHasher>
+FComputeBlocksResult
+ComputeBlocksVariableT(FIOReader& Reader, const FComputeBlocksParams& Params)
+{
+	const uint32 BlockSize = Params.BlockSize;
+
+	const uint64 InputSize = Reader.GetSize();
+
+	const uint32 MinimumBlockSize = ComputeMinVariableBlockSize(BlockSize);
+	const uint32 MaximumBlockSize = ComputeMaxVariableBlockSize(BlockSize);
+
+	const uint64 BytesPerTask = std::min<uint64>(256_MB, std::max<uint64>(InputSize, 1ull));  // TODO: handle task boundary overlap
+	const uint64 NumTasks	  = DivUp(InputSize, BytesPerTask);
+
+	const uint64 TargetMacroBlockSize	 = Params.bNeedMacroBlocks ? Params.MacroBlockTargetSize : 0;
+	const uint64 MinimumMacroBlockSize	 = std::max<uint64>(MinimumBlockSize, TargetMacroBlockSize / 8);
+	const uint64 MaximumMacroBlockSize	 = Params.bNeedMacroBlocks ? Params.MacroBlockMaxSize : 0;
+	const uint32 BlocksPerMacroBlock	 = CheckedNarrow(DivUp(TargetMacroBlockSize - MinimumMacroBlockSize, BlockSize));
+	const uint32 MacroBlockHashThreshold = BlocksPerMacroBlock ? (0xFFFFFFFF / BlocksPerMacroBlock) : 0;
+
+	FBlockSourceInfo SourceInfo;
+	SourceInfo.TotalSize = InputSize;
+
+	struct FTask
+	{
+		uint64			   Offset = 0;
+		FGenericBlockArray Blocks;
+		FGenericBlockArray MacroBlocks;
+	};
+
+	std::vector<FTask> Tasks;
+	Tasks.resize(NumTasks);
+
+	FSchedulerSemaphore IoSemaphore(*GScheduler, 16);
+	FTaskGroup			TaskGroup = GScheduler->CreateTaskGroup();
+
+	FBufferPool BufferPool(BytesPerTask);
+
+	for (uint64 TaskIndex = 0; TaskIndex < NumTasks; ++TaskIndex)
+	{
+		IoSemaphore.Acquire();	// TODO: remove
+
+		const uint64 ThisTaskOffset = BytesPerTask * TaskIndex;
+		const uint64 ThisTaskSize	= CalcChunkSize(TaskIndex, BytesPerTask, InputSize);
+
+		Tasks[TaskIndex].Offset = ThisTaskOffset;
+
+		FBuffer* ScanTaskBuffer = BufferPool.Acquire();
+		UNSYNC_ASSERT(ScanTaskBuffer->Size() >= ThisTaskSize);
+
+		uint64 ReadBytesForTask = BlockingReadLarge(Reader, ThisTaskOffset, ThisTaskSize, ScanTaskBuffer->Data(), ThisTaskSize);
+
+		if (ReadBytesForTask != ThisTaskSize)
+		{
+			UNSYNC_FATAL(L"Expected to read %lld bytes from input, but %lld was actually read.", ThisTaskSize, ReadBytesForTask);
+		}
+
+		auto ScanTask = [&Tasks,
+						 &IoSemaphore,
+						 &BufferPool,
+						 &Params,
+						 &SourceInfo,
+						 MinimumBlockSize,
+						 MaximumBlockSize,
+						 ScanTaskBuffer,
+						 TaskIndex,
+						 ThisTaskSize,
+						 TargetMacroBlockSize,
+						 MinimumMacroBlockSize,
+						 MaximumMacroBlockSize,
+						 MacroBlockHashThreshold]()
+		{
+			FTask& Task = Tasks[TaskIndex];
+
+			const uint8* DataBegin	  = ScanTaskBuffer->Data();
+			const uint8* DataEnd	  = DataBegin + ThisTaskSize;
+			const uint8* LastBlockEnd = DataBegin;
+
+			FGenericBlock CurrentMacroBlock;
+			CurrentMacroBlock.HashStrong.Type = EHashType::Blake3_160;
+			CurrentMacroBlock.Offset		  = Task.Offset;
+
+			blake3_hasher MacroBlockHasher;
+			blake3_hasher_init(&MacroBlockHasher);
+
+			auto ScanFn = [MacroBlockHashThreshold,
+						   MaximumMacroBlockSize,
+						   MaximumBlockSize,
+						   MinimumMacroBlockSize,
+						   DataEnd,
+						   &LastBlockEnd,
+						   &Task,
+						   DataBegin,
+						   &Params,
+						   TargetMacroBlockSize,
+						   &MacroBlockHasher,
+						   &CurrentMacroBlock,
+						   &SourceInfo](const uint8* WindowBegin, const uint8* WindowEnd, uint32 WindowHash) UNSYNC_ATTRIB_FORCEINLINE
+			{
+				// WARNING: Changing this invalidates some of the previously cached blocks.
+				// TODO: compute based on target average block size
+				const uint32 ChunkWindowHashThreshold = 0x20000;
+
+				const bool	 bLastBlock	   = WindowEnd == DataEnd;
+				const uint64 ThisBlockSize = uint64(WindowEnd - LastBlockEnd);
+
+				if (ThisBlockSize >= MaximumBlockSize || WindowHash < ChunkWindowHashThreshold || bLastBlock)
+				{
+					FGenericBlock Block;
+					Block.Offset	 = Task.Offset + uint64(LastBlockEnd - DataBegin);
+					Block.Size		 = CheckedNarrow(ThisBlockSize);
+					Block.HashWeak	 = WindowHash;
+					Block.HashStrong = ComputeHash(LastBlockEnd, ThisBlockSize, Params.Algorithm.StrongHashAlgorithmId);
+
+					if (TargetMacroBlockSize)
+					{
+						blake3_hasher_update(&MacroBlockHasher, LastBlockEnd, ThisBlockSize);
+						CurrentMacroBlock.Size += Block.Size;
+
+						uint32 HashStrong32 = 0;
+						memcpy(&HashStrong32, Block.HashStrong.Data, 4);
+
+						if ((CurrentMacroBlock.Size >= MinimumMacroBlockSize && HashStrong32 < MacroBlockHashThreshold) ||
+							(CurrentMacroBlock.Size + Block.Size > MaximumMacroBlockSize) || bLastBlock)
+						{
+							// Commit the macro block
+							blake3_hasher_finalize(&MacroBlockHasher,
+												   CurrentMacroBlock.HashStrong.Data,
+												   sizeof(CurrentMacroBlock.HashStrong.Data));
+							Task.MacroBlocks.push_back(CurrentMacroBlock);
+
+							if (Params.OnMacroBlockGenerated)
+							{
+								FBufferView BlockView;
+								BlockView.Data = (LastBlockEnd + Block.Size) - CurrentMacroBlock.Size;
+								BlockView.Size = CurrentMacroBlock.Size;
+								Params.OnMacroBlockGenerated(CurrentMacroBlock, SourceInfo, BlockView);
+							}
+
+							// Reset macro block state
+							blake3_hasher_init(&MacroBlockHasher);
+							CurrentMacroBlock.Offset += CurrentMacroBlock.Size;
+							CurrentMacroBlock.Size = 0;
+						}
+					}
+
+					if (Params.OnBlockGenerated)
+					{
+						FBufferView BlockView;
+						BlockView.Data = LastBlockEnd;
+						BlockView.Size = Block.Size;
+						Params.OnBlockGenerated(Block, SourceInfo, BlockView);
+					}
+
+					if (!Task.Blocks.empty())
+					{
+						UNSYNC_ASSERT(Task.Blocks.back().Offset + Task.Blocks.back().Size == Block.Offset);
+					}
+
+					Task.Blocks.push_back(Block);
+
+					LastBlockEnd = WindowEnd;
+
+					return true;
+				}
+				else
+				{
+					return false;
+				}
+			};
+
+			HashScan<WeakHasher>(DataBegin, ThisTaskSize, MinimumBlockSize, ScanFn);
+
+			BufferPool.Release(ScanTaskBuffer);
+			IoSemaphore.Release();
+		};
+
+		TaskGroup.run(ScanTask);
+	}
+
+	TaskGroup.wait();
+
+	// Merge blocks for all the tasks
+
+	FComputeBlocksResult Result;
+
+	for (uint64 I = 0; I < NumTasks; ++I)
+	{
+		const FTask& Task = Tasks[I];
+		for (uint64 J = 0; J < Task.Blocks.size(); ++J)
+		{
+			Result.Blocks.push_back(Task.Blocks[J]);
+		}
+	}
+
+	if (Params.bNeedMacroBlocks)
+	{
+		for (uint64 I = 0; I < NumTasks; ++I)
+		{
+			const FTask& Task = Tasks[I];
+			for (uint64 J = 0; J < Task.MacroBlocks.size(); ++J)
+			{
+				Result.MacroBlocks.push_back(Task.MacroBlocks[J]);
+			}
+		}
+	}
+
+	uint64 UniqueBlockTotalSize = 0;
+	uint64 UniqueBlockMinSize	= ~0ull;
+	uint64 UniqueBlockMaxSize	= 0ull;
+
+	uint64 NumTinyBlocks   = 0;
+	uint64 NumSmallBlocks  = 0;
+	uint64 NumMediumBlocks = 0;
+	uint64 NumLargeBlocks  = 0;
+
+	uint64 NumTotalBlocks = 0;
+
+	THashSet<FGenericHash> UniqueBlockSet;
+	FGenericBlockArray	   UniqueBlocks;
+	for (const FGenericBlock& It : Result.Blocks)
+	{
+		auto InsertResult = UniqueBlockSet.insert(It.HashStrong);
+		if (InsertResult.second)
+		{
+			if (It.Offset + It.Size < InputSize || Result.Blocks.size() == 1)
+			{
+				UniqueBlockMinSize = std::min<uint64>(UniqueBlockMinSize, It.Size);
+			}
+
+			UniqueBlockMaxSize = std::max<uint64>(UniqueBlockMaxSize, It.Size);
+			UniqueBlockTotalSize += It.Size;
+			UniqueBlocks.push_back(It);
+		}
+
+		if (It.Size < MaximumBlockSize / 8)
+		{
+			NumTinyBlocks++;
+		}
+		else if (It.Size <= MaximumBlockSize / 4)
+		{
+			NumSmallBlocks++;
+		}
+		else if (It.Size <= MaximumBlockSize / 2)
+		{
+			NumMediumBlocks++;
+		}
+		else
+		{
+			NumLargeBlocks++;
+		}
+
+		++NumTotalBlocks;
+	}
+
+	double AverageBlockSize = InputSize ? double(UniqueBlockTotalSize / UniqueBlocks.size()) : 0;
+
+	UNSYNC_VERBOSE2(
+		L"Blocks (tiny/small/medium/large): %llu / %llu / %llu / %llu, average unique size: %llu bytes, unique count: %llu, total count: "
+		L"%llu",
+		NumTinyBlocks,
+		NumSmallBlocks,
+		NumMediumBlocks,
+		NumLargeBlocks,
+		(uint64)AverageBlockSize,
+		(uint64)UniqueBlocks.size(),
+		NumTotalBlocks);
+
+	UNSYNC_ASSERT(NumTotalBlocks == Result.Blocks.size());
+
+	return Result;
+}
+
+FComputeBlocksResult
+ComputeBlocksVariable(FIOReader& Reader, const FComputeBlocksParams& Params)
+{
+	switch (Params.Algorithm.WeakHashAlgorithmId)
+	{
+		case EWeakHashAlgorithmID::Naive:
+			return ComputeBlocksVariableT<FRollingChecksum>(Reader, Params);
+		case EWeakHashAlgorithmID::BuzHash:
+			return ComputeBlocksVariableT<FBuzHash>(Reader, Params);
+		default:
+			UNSYNC_FATAL(L"Unsupported weak hash algorithm mode");
+			return {};
+	}
+}
+
+FGenericBlockArray
+ComputeBlocks(FIOReader& Reader, uint32 BlockSize, FAlgorithmOptions Algorithm)
+{
+	FComputeBlocksParams Params;
+	Params.Algorithm			= Algorithm;
+	Params.BlockSize			= BlockSize;
+	FComputeBlocksResult Result = ComputeBlocks(Reader, Params);
+	return std::move(Result.Blocks);
+}
+
+FGenericBlockArray
+ComputeBlocks(const uint8* Data, uint64 Size, uint32 BlockSize, FAlgorithmOptions Algorithm)
+{
+	FComputeBlocksParams Params;
+	Params.Algorithm			= Algorithm;
+	Params.BlockSize			= BlockSize;
+	FComputeBlocksResult Result = ComputeBlocks(Data, Size, Params);
+	return std::move(Result.Blocks);
+}
+
+FGenericBlockArray
+ComputeBlocksVariable(FIOReader& Reader, uint32 BlockSize, EWeakHashAlgorithmID WeakHasher, EStrongHashAlgorithmID StrongHasher)
+{
+	FComputeBlocksParams Params;
+	Params.Algorithm.WeakHashAlgorithmId   = WeakHasher;
+	Params.Algorithm.StrongHashAlgorithmId = StrongHasher;
+	Params.BlockSize					   = BlockSize;
+	FComputeBlocksResult Result			   = ComputeBlocks(Reader, Params);
+	return std::move(Result.Blocks);
+}
+
+template<typename WeakHasher>
+FComputeBlocksResult
+ComputeBlocksFixedT(FIOReader& Reader, const FComputeBlocksParams& Params)
+{
+	UNSYNC_LOG_INDENT;
+
+	auto TimeBegin = TimePointNow();
+
+	const uint32 BlockSize = Params.BlockSize;
+	const uint64 NumBlocks = DivUp(Reader.GetSize(), BlockSize);
+
+	FGenericBlockArray Blocks(NumBlocks);
+	for (uint64 I = 0; I < NumBlocks; ++I)
+	{
+		uint64 ChunkSize = CalcChunkSize(I, BlockSize, Reader.GetSize());
+		Blocks[I].Offset = I * BlockSize;
+		Blocks[I].Size	 = CheckedNarrow(ChunkSize);
+	}
+
+	uint64 ReadSize = std::max<uint64>(BlockSize, 8_MB);
+	if (Params.bNeedMacroBlocks)
+	{
+		UNSYNC_FATAL(L"Macro block generation is not implemented for fixed block mode");
+		ReadSize = std::max<uint64>(ReadSize, Params.MacroBlockTargetSize);
+	}
+	UNSYNC_ASSERT(ReadSize % BlockSize == 0);
+
+	const uint64		NumReads		   = DivUp(Reader.GetSize(), ReadSize);
+	std::atomic<uint64> NumReadsCompleted  = {};
+	std::atomic<uint64> NumBlocksCompleted = {};
+
+	{
+		FSchedulerSemaphore IoSemaphore(*GScheduler, 16);
+		FTaskGroup			TaskGroup = GScheduler->CreateTaskGroup(&IoSemaphore);
+
+		for (uint64 I = 0; I < NumReads; ++I)
+		{
+			uint64 ThisReadSize = CalcChunkSize(I, ReadSize, Reader.GetSize());
+			uint64 Offset		= I * ReadSize;
+
+			auto ReadCallback = [&NumReadsCompleted, &TaskGroup, &NumBlocksCompleted, &Blocks, &Params, BlockSize](FIOBuffer CmdBuffer,
+																												   uint64	 CmdOffset,
+																												   uint64	 CmdReadSize,
+																												   uint64	 CmdUserData)
+			{
+				UNSYNC_ASSERT(CmdReadSize);
+
+				TaskGroup.run(
+					[&NumReadsCompleted,
+					 &NumBlocksCompleted,
+					 &Blocks,
+					 &Params,
+					 BlockSize,
+					 CmdBuffer	= MakeShared(std::move(CmdBuffer)),
+					 BufferSize = CmdReadSize,
+					 Offset		= CmdOffset]()
+					{
+						UNSYNC_ASSERT(CmdBuffer->GetSize() == BufferSize);
+
+						uint8* Buffer = CmdBuffer->GetData();
+
+						UNSYNC_ASSERT(Offset % BlockSize == 0);
+						UNSYNC_ASSERT(BufferSize);
+						UNSYNC_ASSERT(Buffer);
+
+						uint64 FirstBlock	  = Offset / BlockSize;
+						uint64 NumLocalBlocks = DivUp(BufferSize, BlockSize);
+						for (uint64 I = 0; I < NumLocalBlocks; ++I)
+						{
+							FGenericBlock& Block = Blocks[FirstBlock + I];
+
+							UNSYNC_ASSERT(Block.HashWeak == 0);
+							UNSYNC_ASSERT(Block.HashStrong == FGenericHash{});
+
+							Block.HashStrong = ComputeHash(Buffer + I * BlockSize, Block.Size, Params.Algorithm.StrongHashAlgorithmId);
+
+							WeakHasher HashWeak;
+							HashWeak.Update(Buffer + I * BlockSize, Block.Size);
+							Block.HashWeak = HashWeak.Get();
+
+							++NumBlocksCompleted;
+						}
+
+						++NumReadsCompleted;
+					});
+			};
+
+			Reader.ReadAsync(Offset, ThisReadSize, 0, ReadCallback);
+		}
+		Reader.FlushAll();
+		TaskGroup.wait();
+	}
+
+	UNSYNC_ASSERT(NumReadsCompleted == NumReads);
+	UNSYNC_ASSERT(NumBlocksCompleted == NumBlocks);
+
+	md5_context Hasher;
+	md5_init(&Hasher);
+	for (uint64 I = 0; I < NumBlocks; ++I)
+	{
+		if (Blocks[I].HashStrong == FGenericHash{})
+		{
+			UNSYNC_ERROR(L"Found invalid hash in block %llu", I);
+		}
+		UNSYNC_ASSERT(Blocks[I].HashStrong != FGenericHash{});
+		md5_update(&Hasher, Blocks[I].HashStrong.Data, sizeof(Blocks[I].HashStrong));
+	}
+	uint8 Hash[16] = {};
+	md5_finish(&Hasher, Hash);
+	std::string HashStr = BytesToHexString(Hash, sizeof(Hash));
+	UNSYNC_VERBOSE2(L"Hash: %hs", HashStr.c_str());
+
+	uint64 TocSize = sizeof(FBlock128) * NumBlocks;
+	UNSYNC_VERBOSE2(L"Manifest size: %lld bytes (%.2f MB), blocks: %d", (long long)TocSize, SizeMb(TocSize), uint32(NumBlocks));
+
+	double Duration = DurationSec(TimeBegin, TimePointNow());
+	UNSYNC_VERBOSE2(L"Done in %.3f sec (%.3f MB / sec)", Duration, SizeMb((double(Reader.GetSize()) / Duration)));
+
+	THashSet<uint32> UniqueHashes;
+	for (const auto& It : Blocks)
+	{
+		UniqueHashes.insert(It.HashWeak);
+	}
+
+	FComputeBlocksResult Result;
+
+	std::swap(Result.Blocks, Blocks);
+
+	return Result;
+}
+
+FComputeBlocksResult
+ComputeBlocksFixed(FIOReader& Reader, const FComputeBlocksParams& Params)
+{
+	switch (Params.Algorithm.WeakHashAlgorithmId)
+	{
+		case EWeakHashAlgorithmID::Naive:
+			return ComputeBlocksFixedT<FRollingChecksum>(Reader, Params);
+		case EWeakHashAlgorithmID::BuzHash:
+			return ComputeBlocksFixedT<FBuzHash>(Reader, Params);
+		default:
+			UNSYNC_FATAL(L"Unsupported weak hash algorithm mode");
+			return {};
+	}
+}
+
+FComputeBlocksResult
+ComputeBlocks(FIOReader& Reader, const FComputeBlocksParams& Params)
+{
+	switch (Params.Algorithm.ChunkingAlgorithmId)
+	{
+		case EChunkingAlgorithmID::FixedBlocks:
+			return ComputeBlocksFixed(Reader, Params);
+		case EChunkingAlgorithmID::VariableBlocks:
+			return ComputeBlocksVariable(Reader, Params);
+		default:
+			UNSYNC_FATAL(L"Unsupported chunking mode");
+			return {};
+	}
+}
+
+FComputeBlocksResult
+ComputeBlocks(const uint8* Data, uint64 Size, const FComputeBlocksParams& Params)
+{
+	FMemReader DataReader(Data, Size);
+	return ComputeBlocks(DataReader, Params);
+}
+
+}  // namespace unsync
diff --git a/Engine/Source/Programs/Unsync/Private/UnsyncChunking.h b/Engine/Source/Programs/Unsync/Private/UnsyncChunking.h
new file mode 100644
--- /dev/null
+++ b/Engine/Source/Programs/Unsync/Private/UnsyncChunking.h
@@ -0,0 +1,25 @@
+// Copyright Epic Games, Inc. All Rights Reserved.
+
+#include "UnsyncCommon.h"
+#include "UnsyncCore.h"
+
+namespace unsync {
+
+struct FComputeBlocksResult
+{
+	FGenericBlockArray Blocks;
+	FGenericBlockArray MacroBlocks;
+};
+
+FComputeBlocksResult ComputeBlocks(FIOReader& Reader, const FComputeBlocksParams& Params);
+FComputeBlocksResult ComputeBlocks(const uint8* Data, uint64 Size, const FComputeBlocksParams& Params);
+FComputeBlocksResult ComputeBlocksVariable(FIOReader& Reader, const FComputeBlocksParams& Params);
+
+FGenericBlockArray ComputeBlocks(FIOReader& Reader, uint32 BlockSize, FAlgorithmOptions Algorithm);
+FGenericBlockArray ComputeBlocks(const uint8* Data, uint64 Size, uint32 BlockSize, FAlgorithmOptions Algorithm);
+FGenericBlockArray ComputeBlocksVariable(FIOReader&				Reader,
+										 uint32					BlockSize,
+										 EWeakHashAlgorithmID	WeakHasher,
+										 EStrongHashAlgorithmID StrongHasher);
+
+}  // namespace unsync
diff --git a/Engine/Source/Programs/Unsync/Private/UnsyncCmdDiff.cpp b/Engine/Source/Programs/Unsync/Private/UnsyncCmdDiff.cpp
--- a/Engine/Source/Programs/Unsync/Private/UnsyncCmdDiff.cpp
+++ b/Engine/Source/Programs/Unsync/Private/UnsyncCmdDiff.cpp
@@ -3,6 +3,7 @@
 #include "UnsyncCmdDiff.h"
 #include "UnsyncFile.h"
 #include "UnsyncTarget.h"
+#include "UnsyncDiff.h"
 
 namespace unsync {
 
diff --git a/Engine/Source/Programs/Unsync/Private/UnsyncCmdHash.cpp b/Engine/Source/Programs/Unsync/Private/UnsyncCmdHash.cpp
--- a/Engine/Source/Programs/Unsync/Private/UnsyncCmdHash.cpp
+++ b/Engine/Source/Programs/Unsync/Private/UnsyncCmdHash.cpp
@@ -5,6 +5,7 @@
 #include "UnsyncFile.h"
 #include "UnsyncSerialization.h"
 #include "UnsyncPack.h"
+#include "UnsyncChunking.h"
 
 #include <memory>
 
diff --git a/Engine/Source/Programs/Unsync/Private/UnsyncCmdPack.cpp b/Engine/Source/Programs/Unsync/Private/UnsyncCmdPack.cpp
--- a/Engine/Source/Programs/Unsync/Private/UnsyncCmdPack.cpp
+++ b/Engine/Source/Programs/Unsync/Private/UnsyncCmdPack.cpp
@@ -10,6 +10,7 @@
 #include "UnsyncProgress.h"
 #include "UnsyncScheduler.h"
 #include "UnsyncPack.h"
+#include "UnsyncChunking.h"
 
 #include <stdio.h>
 #include <stdlib.h>
diff --git a/Engine/Source/Programs/Unsync/Private/UnsyncCore.cpp b/Engine/Source/Programs/Unsync/Private/UnsyncCore.cpp
--- a/Engine/Source/Programs/Unsync/Private/UnsyncCore.cpp
+++ b/Engine/Source/Programs/Unsync/Private/UnsyncCore.cpp
@@ -2,20 +2,23 @@
 
 #include "UnsyncCommon.h"
 
+#include "UnsyncChunking.h"
 #include "UnsyncCompression.h"
 #include "UnsyncCore.h"
+#include "UnsyncDiff.h"
 #include "UnsyncFile.h"
 #include "UnsyncHashTable.h"
+#include "UnsyncHttp.h"
 #include "UnsyncProgress.h"
 #include "UnsyncProxy.h"
 #include "UnsyncScan.h"
 #include "UnsyncScavenger.h"
+#include "UnsyncScheduler.h"
 #include "UnsyncSerialization.h"
+#include "UnsyncTarget.h"
 #include "UnsyncThread.h"
 #include "UnsyncUtil.h"
-#include "UnsyncTarget.h"
-#include "UnsyncHttp.h"
-#include "UnsyncScheduler.h"
+#include "UnsyncVersion.h"
 
 #include <condition_variable>
 #include <filesystem>
@@ -30,71 +33,10 @@ UNSYNC_THIRD_PARTY_INCLUDES_START
 #include <md5-sse2.h>
 UNSYNC_THIRD_PARTY_INCLUDES_END
 
-#define UNSYNC_VERSION_STR "1.0.78"
-
 namespace unsync {
 
 bool GDryRun = false;
 
-inline uint32
-ComputeMinVariableBlockSize(uint32 BlockSize)
-{
-	return std::max(BlockSize, 4096u) / 2;	// changing this invalidates cached blocks
-}
-
-inline uint32
-ComputeMaxVariableBlockSize(uint32 BlockSize)
-{
-	return std::max(BlockSize, 4096u) * 4;	// changing this invalidates cached blocks
-}
-
-static uint64
-BlockingReadLarge(FIOReader& Reader, uint64 Offset, uint64 Size, uint8* OutputBuffer, uint64 OutputBufferSize)
-{
-	const uint64 BytesPerRead = 2_MB;
-	const uint64 ReadEnd	  = std::min(Offset + Size, Reader.GetSize());
-	const uint64 ClampedSize  = ReadEnd - Offset;
-
-	std::atomic<uint64> TotalReadSize = 0;
-
-	if (ClampedSize == 0)
-	{
-		return TotalReadSize;
-	}
-
-	FSchedulerSemaphore IoSemaphore(*GScheduler, 16);
-	FTaskGroup			CopyTasks = GScheduler->CreateTaskGroup(&IoSemaphore);
-
-	uint64 NumReads = DivUp(ClampedSize, BytesPerRead);
-	for (uint64 ReadIndex = 0; ReadIndex < NumReads; ++ReadIndex)
-	{
-		const uint64 ThisBatchSize	= CalcChunkSize(ReadIndex, BytesPerRead, ClampedSize);
-		const uint64 OutputOffset	= BytesPerRead * ReadIndex;
-		const uint64 ThisReadOffset = Offset + OutputOffset;
-
-		auto ReadCallback = [OutputBuffer, OutputBufferSize, &TotalReadSize, &CopyTasks](FIOBuffer CmdBuffer,
-																									   uint64	 CmdSourceOffset,
-																									   uint64	 CmdReadSize,
-																									   uint64	 OutputOffset) {
-			UNSYNC_ASSERT(OutputOffset + CmdReadSize <= OutputBufferSize);
-
-			CopyTasks.run(
-				[OutputBuffer, OutputOffset, CmdReadSize, CmdBuffer = MakeShared(std::move(CmdBuffer)), &TotalReadSize]()
-				{
-					memcpy(OutputBuffer + OutputOffset, CmdBuffer->GetData(), CmdReadSize);
-					TotalReadSize += CmdReadSize;
-				});
-		};
-
-		Reader.ReadAsync(ThisReadOffset, ThisBatchSize, OutputOffset, ReadCallback);
-	}
-
-	Reader.FlushAll();
-	CopyTasks.wait();
-
-	return TotalReadSize;
-}
-
 FBlock128
 ToBlock128(const FGenericBlock& GenericBlock)
 {
@@ -120,893 +62,6 @@ ToBlock128(FGenericBlockArray& GenericBlocks)
 	return Result;
 }
 
-template<typename WeakHasher>
-FComputeBlocksResult
-ComputeBlocksVariableT(FIOReader& Reader, const FComputeBlocksParams& Params)
-{
-	const uint32 BlockSize = Params.BlockSize;
-
-	const uint64 InputSize = Reader.GetSize();
-
-	const uint32 MinimumBlockSize = ComputeMinVariableBlockSize(BlockSize);
-	const uint32 MaximumBlockSize = ComputeMaxVariableBlockSize(BlockSize);
-
-	const uint64 BytesPerTask = std::min<uint64>(256_MB, std::max<uint64>(InputSize, 1ull));  // TODO: handle task boundary overlap
-	const uint64 NumTasks	  = DivUp(InputSize, BytesPerTask);
-
-	const uint64 TargetMacroBlockSize	 = Params.bNeedMacroBlocks ? Params.MacroBlockTargetSize : 0;
-	const uint64 MinimumMacroBlockSize	 = std::max<uint64>(MinimumBlockSize, TargetMacroBlockSize / 8);
-	const uint64 MaximumMacroBlockSize	 = Params.bNeedMacroBlocks ? Params.MacroBlockMaxSize : 0;
-	const uint32 BlocksPerMacroBlock	 = CheckedNarrow(DivUp(TargetMacroBlockSize - MinimumMacroBlockSize, BlockSize));
-	const uint32 MacroBlockHashThreshold = BlocksPerMacroBlock ? (0xFFFFFFFF / BlocksPerMacroBlock) : 0;
-
-	FBlockSourceInfo SourceInfo;
-	SourceInfo.TotalSize = InputSize;
-
-	struct FTask
-	{
-		uint64			   Offset = 0;
-		FGenericBlockArray Blocks;
-		FGenericBlockArray MacroBlocks;
-	};
-
-	std::vector<FTask> Tasks;
-	Tasks.resize(NumTasks);
-
-	FSchedulerSemaphore IoSemaphore(*GScheduler, 16);
-	FTaskGroup TaskGroup = GScheduler->CreateTaskGroup();
-
-	FBufferPool BufferPool(BytesPerTask);
-
-	for (uint64 TaskIndex = 0; TaskIndex < NumTasks; ++TaskIndex)
-	{
-		IoSemaphore.Acquire(); // TODO: remove
-
-		const uint64 ThisTaskOffset = BytesPerTask * TaskIndex;
-		const uint64 ThisTaskSize	= CalcChunkSize(TaskIndex, BytesPerTask, InputSize);
-
-		Tasks[TaskIndex].Offset = ThisTaskOffset;
-
-		FBuffer* ScanTaskBuffer = BufferPool.Acquire();
-		UNSYNC_ASSERT(ScanTaskBuffer->Size() >= ThisTaskSize);
-
-		uint64 ReadBytesForTask = BlockingReadLarge(Reader, ThisTaskOffset, ThisTaskSize, ScanTaskBuffer->Data(), ThisTaskSize);
-
-		if (ReadBytesForTask != ThisTaskSize)
-		{
-			UNSYNC_FATAL(L"Expected to read %lld bytes from input, but %lld was actually read.", ThisTaskSize, ReadBytesForTask);
-		}
-
-		auto ScanTask = [&Tasks,
-						 &IoSemaphore,
-						 &BufferPool,
-						 &Params,
-						 &SourceInfo,
-						 MinimumBlockSize,
-						 MaximumBlockSize,
-						 ScanTaskBuffer,
-						 TaskIndex,
-						 ThisTaskSize,
-						 TargetMacroBlockSize,
-						 MinimumMacroBlockSize,
-						 MaximumMacroBlockSize,
-						 MacroBlockHashThreshold]() {
-			FTask& Task = Tasks[TaskIndex];
-
-			const uint8* DataBegin	  = ScanTaskBuffer->Data();
-			const uint8* DataEnd	  = DataBegin + ThisTaskSize;
-			const uint8* LastBlockEnd = DataBegin;
-
-			FGenericBlock CurrentMacroBlock;
-			CurrentMacroBlock.HashStrong.Type = EHashType::Blake3_160;
-			CurrentMacroBlock.Offset		  = Task.Offset;
-
-			blake3_hasher MacroBlockHasher;
-			blake3_hasher_init(&MacroBlockHasher);
-
-			auto ScanFn = [MacroBlockHashThreshold,
-						   MaximumMacroBlockSize,
-						   MaximumBlockSize,
-						   MinimumMacroBlockSize,
-						   DataEnd,
-						   &LastBlockEnd,
-						   &Task,
-						   DataBegin,
-						   &Params,
-						   TargetMacroBlockSize,
-						   &MacroBlockHasher,
-						   &CurrentMacroBlock,
-						   &SourceInfo](const uint8* WindowBegin, const uint8* WindowEnd, uint32 WindowHash)
-							  UNSYNC_ATTRIB_FORCEINLINE {
-								  // WARNING: Changing this invalidates some of the previously cached blocks.
-								  // TODO: compute based on target average block size
-								  const uint32 ChunkWindowHashThreshold = 0x20000;
-
-								  const bool   bLastBlock	 = WindowEnd == DataEnd;
-								  const uint64 ThisBlockSize = uint64(WindowEnd - LastBlockEnd);
-
-								  if (ThisBlockSize >= MaximumBlockSize || WindowHash < ChunkWindowHashThreshold || bLastBlock)
-								  {
-									  FGenericBlock Block;
-									  Block.Offset	   = Task.Offset + uint64(LastBlockEnd - DataBegin);
-									  Block.Size	   = CheckedNarrow(ThisBlockSize);
-									  Block.HashWeak   = WindowHash;
-									  Block.HashStrong = ComputeHash(LastBlockEnd, ThisBlockSize, Params.Algorithm.StrongHashAlgorithmId);
-
-									  if (TargetMacroBlockSize)
-									  {
-										  blake3_hasher_update(&MacroBlockHasher, LastBlockEnd, ThisBlockSize);
-										  CurrentMacroBlock.Size += Block.Size;
-
-										  uint32 HashStrong32 = 0;
-										  memcpy(&HashStrong32, Block.HashStrong.Data, 4);
-
-										  if ((CurrentMacroBlock.Size >= MinimumMacroBlockSize && HashStrong32 < MacroBlockHashThreshold) ||
-											  (CurrentMacroBlock.Size + Block.Size > MaximumMacroBlockSize) || bLastBlock)
-										  {
-											  // Commit the macro block
-											  blake3_hasher_finalize(&MacroBlockHasher,
-																	 CurrentMacroBlock.HashStrong.Data,
-																	 sizeof(CurrentMacroBlock.HashStrong.Data));
-											  Task.MacroBlocks.push_back(CurrentMacroBlock);
-
-											  if (Params.OnMacroBlockGenerated)
-											  {
-												  FBufferView BlockView;
-												  BlockView.Data = (LastBlockEnd + Block.Size) - CurrentMacroBlock.Size;
-												  BlockView.Size = CurrentMacroBlock.Size;
-												  Params.OnMacroBlockGenerated(CurrentMacroBlock, SourceInfo, BlockView);
-											  }
-
-											  // Reset macro block state
-											  blake3_hasher_init(&MacroBlockHasher);
-											  CurrentMacroBlock.Offset += CurrentMacroBlock.Size;
-											  CurrentMacroBlock.Size = 0;
-										  }
-									  }
-
-									  if (Params.OnBlockGenerated)
-									  {
-										  FBufferView BlockView;
-										  BlockView.Data = LastBlockEnd;
-										  BlockView.Size = Block.Size;
-										  Params.OnBlockGenerated(Block, SourceInfo, BlockView);
-									  }
-
-									  if (!Task.Blocks.empty())
-									  {
-										  UNSYNC_ASSERT(Task.Blocks.back().Offset + Task.Blocks.back().Size == Block.Offset);
-									  }
-
-									  Task.Blocks.push_back(Block);
-
-									  LastBlockEnd = WindowEnd;
-
-									  return true;
-								  }
-								  else
-								  {
-									  return false;
-								  }
-							  };
-
-			HashScan<WeakHasher>(DataBegin, ThisTaskSize, MinimumBlockSize, ScanFn);
-
-			BufferPool.Release(ScanTaskBuffer);
-			IoSemaphore.Release();
-		};
-
-		TaskGroup.run(ScanTask);
-	}
-
-	TaskGroup.wait();
-
-	// Merge blocks for all the tasks
-
-	FComputeBlocksResult Result;
-
-	for (uint64 I = 0; I < NumTasks; ++I)
-	{
-		const FTask& Task = Tasks[I];
-		for (uint64 J = 0; J < Task.Blocks.size(); ++J)
-		{
-			Result.Blocks.push_back(Task.Blocks[J]);
-		}
-	}
-
-	if (Params.bNeedMacroBlocks)
-	{
-		for (uint64 I = 0; I < NumTasks; ++I)
-		{
-			const FTask& Task = Tasks[I];
-			for (uint64 J = 0; J < Task.MacroBlocks.size(); ++J)
-			{
-				Result.MacroBlocks.push_back(Task.MacroBlocks[J]);
-			}
-		}
-	}
-
-	uint64 UniqueBlockTotalSize = 0;
-	uint64 UniqueBlockMinSize	= ~0ull;
-	uint64 UniqueBlockMaxSize	= 0ull;
-
-	uint64 NumTinyBlocks   = 0;
-	uint64 NumSmallBlocks  = 0;
-	uint64 NumMediumBlocks = 0;
-	uint64 NumLargeBlocks  = 0;
-
-	uint64 NumTotalBlocks = 0;
-
-	THashSet<FGenericHash> UniqueBlockSet;
-	FGenericBlockArray	   UniqueBlocks;
-	for (const FGenericBlock& It : Result.Blocks)
-	{
-		auto InsertResult = UniqueBlockSet.insert(It.HashStrong);
-		if (InsertResult.second)
-		{
-			if (It.Offset + It.Size < InputSize || Result.Blocks.size() == 1)
-			{
-				UniqueBlockMinSize = std::min<uint64>(UniqueBlockMinSize, It.Size);
-			}
-
-			UniqueBlockMaxSize = std::max<uint64>(UniqueBlockMaxSize, It.Size);
-			UniqueBlockTotalSize += It.Size;
-			UniqueBlocks.push_back(It);
-		}
-
-		if (It.Size < MaximumBlockSize / 8)
-		{
-			NumTinyBlocks++;
-		}
-		else if (It.Size <= MaximumBlockSize / 4)
-		{
-			NumSmallBlocks++;
-		}
-		else if (It.Size <= MaximumBlockSize / 2)
-		{
-			NumMediumBlocks++;
-		}
-		else
-		{
-			NumLargeBlocks++;
-		}
-
-		++NumTotalBlocks;
-	}
-
-	double AverageBlockSize = InputSize ? double(UniqueBlockTotalSize / UniqueBlocks.size()) : 0;
-
-	UNSYNC_VERBOSE2(
-		L"Blocks (tiny/small/medium/large): %llu / %llu / %llu / %llu, average unique size: %llu bytes, unique count: %llu, total count: "
-		L"%llu",
-		NumTinyBlocks,
-		NumSmallBlocks,
-		NumMediumBlocks,
-		NumLargeBlocks,
-		(uint64)AverageBlockSize,
-		(uint64)UniqueBlocks.size(),
-		NumTotalBlocks);
-
-	UNSYNC_ASSERT(NumTotalBlocks == Result.Blocks.size());
-
-	return Result;
-}
-
-FComputeBlocksResult
-ComputeBlocksVariable(FIOReader& Reader, const FComputeBlocksParams& Params)
-{
-	switch (Params.Algorithm.WeakHashAlgorithmId)
-	{
-		case EWeakHashAlgorithmID::Naive:
-			return ComputeBlocksVariableT<FRollingChecksum>(Reader, Params);
-		case EWeakHashAlgorithmID::BuzHash:
-			return ComputeBlocksVariableT<FBuzHash>(Reader, Params);
-		default:
-			UNSYNC_FATAL(L"Unsupported weak hash algorithm mode");
-			return {};
-	}
-}
-
-FGenericBlockArray
-ComputeBlocks(FIOReader& Reader, uint32 BlockSize, FAlgorithmOptions Algorithm)
-{
-	FComputeBlocksParams Params;
-	Params.Algorithm			= Algorithm;
-	Params.BlockSize			= BlockSize;
-	FComputeBlocksResult Result = ComputeBlocks(Reader, Params);
-	return std::move(Result.Blocks);
-}
-
-FGenericBlockArray
-ComputeBlocks(const uint8* Data, uint64 Size, uint32 BlockSize, FAlgorithmOptions Algorithm)
-{
-	FComputeBlocksParams Params;
-	Params.Algorithm			= Algorithm;
-	Params.BlockSize			= BlockSize;
-	FComputeBlocksResult Result = ComputeBlocks(Data, Size, Params);
-	return std::move(Result.Blocks);
-}
-
-FGenericBlockArray
-ComputeBlocksVariable(FIOReader& Reader, uint32 BlockSize, EWeakHashAlgorithmID WeakHasher, EStrongHashAlgorithmID StrongHasher)
-{
-	FComputeBlocksParams Params;
-	Params.Algorithm.WeakHashAlgorithmId   = WeakHasher;
-	Params.Algorithm.StrongHashAlgorithmId = StrongHasher;
-	Params.BlockSize					   = BlockSize;
-	FComputeBlocksResult Result			   = ComputeBlocks(Reader, Params);
-	return std::move(Result.Blocks);
-}
-
-template<typename WeakHasher>
-FComputeBlocksResult
-ComputeBlocksFixedT(FIOReader& Reader, const FComputeBlocksParams& Params)
-{
-	UNSYNC_LOG_INDENT;
-
-	auto TimeBegin = TimePointNow();
-
-	const uint32 BlockSize = Params.BlockSize;
-	const uint64 NumBlocks = DivUp(Reader.GetSize(), BlockSize);
-
-	FGenericBlockArray Blocks(NumBlocks);
-	for (uint64 I = 0; I < NumBlocks; ++I)
-	{
-		uint64 ChunkSize = CalcChunkSize(I, BlockSize, Reader.GetSize());
-		Blocks[I].Offset = I * BlockSize;
-		Blocks[I].Size	 = CheckedNarrow(ChunkSize);
-	}
-
-	uint64 ReadSize = std::max<uint64>(BlockSize, 8_MB);
-	if (Params.bNeedMacroBlocks)
-	{
-		UNSYNC_FATAL(L"Macro block generation is not implemented for fixed block mode");
-		ReadSize = std::max<uint64>(ReadSize, Params.MacroBlockTargetSize);
-	}
-	UNSYNC_ASSERT(ReadSize % BlockSize == 0);
-
-	const uint64		NumReads		   = DivUp(Reader.GetSize(), ReadSize);
-	std::atomic<uint64> NumReadsCompleted  = {};
-	std::atomic<uint64> NumBlocksCompleted = {};
-
-	{
-		FSchedulerSemaphore IoSemaphore(*GScheduler, 16);
-		FTaskGroup			TaskGroup = GScheduler->CreateTaskGroup(&IoSemaphore);
-
-		for (uint64 I = 0; I < NumReads; ++I)
-		{
-			uint64 ThisReadSize = CalcChunkSize(I, ReadSize, Reader.GetSize());
-			uint64 Offset		= I * ReadSize;
-
-			auto ReadCallback = [&NumReadsCompleted, &TaskGroup, &NumBlocksCompleted, &Blocks, &Params, BlockSize](
-									FIOBuffer CmdBuffer,
-									uint64	  CmdOffset,
-									uint64	  CmdReadSize,
-									uint64	  CmdUserData) {
-				UNSYNC_ASSERT(CmdReadSize);
-
-				TaskGroup.run([&NumReadsCompleted,
-							   &NumBlocksCompleted,
-							   &Blocks,
-							   &Params,
-							   BlockSize,
-							   CmdBuffer  = MakeShared(std::move(CmdBuffer)),
-							   BufferSize = CmdReadSize,
-							   Offset	  = CmdOffset]() {
-					UNSYNC_ASSERT(CmdBuffer->GetSize() == BufferSize);
-
-					uint8* Buffer = CmdBuffer->GetData();
-
-					UNSYNC_ASSERT(Offset % BlockSize == 0);
-					UNSYNC_ASSERT(BufferSize);
-					UNSYNC_ASSERT(Buffer);
-
-					uint64 FirstBlock	  = Offset / BlockSize;
-					uint64 NumLocalBlocks = DivUp(BufferSize, BlockSize);
-					for (uint64 I = 0; I < NumLocalBlocks; ++I)
-					{
-						FGenericBlock& Block = Blocks[FirstBlock + I];
-
-						UNSYNC_ASSERT(Block.HashWeak == 0);
-						UNSYNC_ASSERT(Block.HashStrong == FGenericHash{});
-
-						Block.HashStrong = ComputeHash(Buffer + I * BlockSize, Block.Size, Params.Algorithm.StrongHashAlgorithmId);
-
-						WeakHasher HashWeak;
-						HashWeak.Update(Buffer + I * BlockSize, Block.Size);
-						Block.HashWeak = HashWeak.Get();
-
-						++NumBlocksCompleted;
-					}
-
-					++NumReadsCompleted;
-				});
-			};
-
-			Reader.ReadAsync(Offset, ThisReadSize, 0, ReadCallback);
-		}
-		Reader.FlushAll();
-		TaskGroup.wait();
-	}
-
-	UNSYNC_ASSERT(NumReadsCompleted == NumReads);
-	UNSYNC_ASSERT(NumBlocksCompleted == NumBlocks);
-
-	md5_context Hasher;
-	md5_init(&Hasher);
-	for (uint64 I = 0; I < NumBlocks; ++I)
-	{
-		if (Blocks[I].HashStrong == FGenericHash{})
-		{
-			UNSYNC_ERROR(L"Found invalid hash in block %llu", I);
-		}
-		UNSYNC_ASSERT(Blocks[I].HashStrong != FGenericHash{});
-		md5_update(&Hasher, Blocks[I].HashStrong.Data, sizeof(Blocks[I].HashStrong));
-	}
-	uint8 Hash[16] = {};
-	md5_finish(&Hasher, Hash);
-	std::string HashStr = BytesToHexString(Hash, sizeof(Hash));
-	UNSYNC_VERBOSE2(L"Hash: %hs", HashStr.c_str());
-
-	uint64 TocSize = sizeof(FBlock128) * NumBlocks;
-	UNSYNC_VERBOSE2(L"Manifest size: %lld bytes (%.2f MB), blocks: %d", (long long)TocSize, SizeMb(TocSize), uint32(NumBlocks));
-
-	double Duration = DurationSec(TimeBegin, TimePointNow());
-	UNSYNC_VERBOSE2(L"Done in %.3f sec (%.3f MB / sec)", Duration, SizeMb((double(Reader.GetSize()) / Duration)));
-
-	THashSet<uint32> UniqueHashes;
-	for (const auto& It : Blocks)
-	{
-		UniqueHashes.insert(It.HashWeak);
-	}
-
-	FComputeBlocksResult Result;
-
-	std::swap(Result.Blocks, Blocks);
-
-	return Result;
-}
-
-FComputeBlocksResult
-ComputeBlocksFixed(FIOReader& Reader, const FComputeBlocksParams& Params)
-{
-	switch (Params.Algorithm.WeakHashAlgorithmId)
-	{
-		case EWeakHashAlgorithmID::Naive:
-			return ComputeBlocksFixedT<FRollingChecksum>(Reader, Params);
-		case EWeakHashAlgorithmID::BuzHash:
-			return ComputeBlocksFixedT<FBuzHash>(Reader, Params);
-		default:
-			UNSYNC_FATAL(L"Unsupported weak hash algorithm mode");
-			return {};
-	}
-}
-
-const std::string&
-GetVersionString()
-{
-	static std::string Result = []() {
-		// TODO: generate a version string based on git state
-		const char* GitRev	  = "";
-		const char* GitBranch = "";
-		// const char* GIT_TAG = nullptr;
-
-		static char Str[256];
-
-		if (strlen(GitBranch) && strlen(GitRev))
-		{
-			snprintf(Str, sizeof(Str), UNSYNC_VERSION_STR " [%s:%s]", GitBranch, GitRev);
-		}
-		else if (strlen(GitRev))
-		{
-			snprintf(Str, sizeof(Str), UNSYNC_VERSION_STR " [%s]", GitRev);
-		}
-		else
-		{
-			snprintf(Str, sizeof(Str), UNSYNC_VERSION_STR);
-		}
-
-		return std::string(Str);
-	}();
-
-	return Result;
-}
-
-FComputeBlocksResult
-ComputeBlocks(FIOReader& Reader, const FComputeBlocksParams& Params)
-{
-	switch (Params.Algorithm.ChunkingAlgorithmId)
-	{
-		case EChunkingAlgorithmID::FixedBlocks:
-			return ComputeBlocksFixed(Reader, Params);
-		case EChunkingAlgorithmID::VariableBlocks:
-			return ComputeBlocksVariable(Reader, Params);
-		default:
-			UNSYNC_FATAL(L"Unsupported chunking mode");
-			return {};
-	}
-}
-
-FComputeBlocksResult
-ComputeBlocks(const uint8* Data, uint64 Size, const FComputeBlocksParams& Params)
-{
-	FMemReader DataReader(Data, Size);
-	return ComputeBlocks(DataReader, Params);
-}
-
-template<typename BlockType>
-bool
-ValidateBlockListT(const std::vector<BlockType>& Blocks)
-{
-	uint64 CurrentOffset = 0;
-	for (const BlockType& Block : Blocks)
-	{
-		if (CurrentOffset != Block.Offset)
-		{
-			UNSYNC_ERROR(L"Found block at unexpected offset. Blocks are expected to be ordered by offset and contiguous.");
-			return false;
-		}
-
-		CurrentOffset += Block.Size;
-	}
-
-	return true;
-}
-
-FNeedList
-DiffBlocksVariable(FIOReader&				 BaseDataReader,
-				   uint32					 BlockSize,
-				   EWeakHashAlgorithmID		 WeakHasher,
-				   EStrongHashAlgorithmID	 StrongHasher,
-				   const FGenericBlockArray& SourceBlocks)
-{
-	FGenericBlockArray BaseBlocks = ComputeBlocksVariable(BaseDataReader, BlockSize, WeakHasher, StrongHasher);
-	if (!ValidateBlockListT(BaseBlocks))
-	{
-		UNSYNC_FATAL(L"Base block list validation failed");
-	}
-
-	return DiffManifestBlocks(SourceBlocks, BaseBlocks);
-}
-
-inline FHash128
-ToHash128(const FHash128& X)
-{
-	return X;
-}
-
-inline FHash128
-ToHash128(const FGenericHash& X)
-{
-	return X.ToHash128();
-}
-
-template<typename BlockType>
-FNeedList
-DiffManifestBlocksT(const std::vector<BlockType>& SourceBlocks, const std::vector<BlockType>& BaseBlocks)
-{
-	FNeedList NeedList;
-
-	struct BlockIndexAndCount
-	{
-		uint64 Index = 0;
-		uint64 Count = 0;
-	};
-
-	THashMap<typename BlockType::StrongHashType, BlockIndexAndCount> BaseBlockMap;
-	THashMap<uint64, uint64>										 BaseBlockByOffset;
-
-	for (uint64 I = 0; I < BaseBlocks.size(); ++I)
-	{
-		const BlockType& Block			= BaseBlocks[I];
-		BaseBlockByOffset[Block.Offset] = I;
-
-		auto Existing = BaseBlockMap.find(Block.HashStrong);
-		if (Existing == BaseBlockMap.end())
-		{
-			BlockIndexAndCount Item;
-			Item.Index = I;
-			Item.Count = 1;
-			BaseBlockMap.insert(std::make_pair(Block.HashStrong, Item));
-		}
-		else
-		{
-			Existing->second.Count += 1;
-		}
-	}
-
-	for (uint64 I = 0; I < SourceBlocks.size(); ++I)
-	{
-		const BlockType& SourceBlock = SourceBlocks[I];
-
-		auto BaseBlockIt = BaseBlockMap.find(SourceBlock.HashStrong);
-		if (BaseBlockIt == BaseBlockMap.end())
-		{
-			FNeedBlock NeedBlock;
-			NeedBlock.Hash		   = SourceBlock.HashStrong;
-			NeedBlock.Size		   = SourceBlock.Size;
-			NeedBlock.SourceOffset = SourceBlock.Offset;
-			NeedBlock.TargetOffset = SourceBlock.Offset;
-			NeedList.Source.push_back(NeedBlock);
-		}
-		else
-		{
-			BlockIndexAndCount IndexAndCount = BaseBlockIt->second;
-			const BlockType&   BaseBlock	 = BaseBlocks[IndexAndCount.Index];
-
-			UNSYNC_ASSERT(BaseBlock.Size == SourceBlock.Size);
-
-			FNeedBlock NeedBlock;
-			NeedBlock.Hash		   = BaseBlock.HashStrong;
-			NeedBlock.Size		   = BaseBlock.Size;
-			NeedBlock.SourceOffset = BaseBlock.Offset;
-			NeedBlock.TargetOffset = SourceBlock.Offset;
-
-			const FNeedBlock* LastBaseNeedBlock = NeedList.Base.empty() ? nullptr : &(NeedList.Base.back());
-
-			// Try to preserve contiguous base data reads
-			if (LastBaseNeedBlock)
-			{
-				uint64 LastBlockEnd			  = LastBaseNeedBlock->SourceOffset + LastBaseNeedBlock->Size;
-				auto   ConsecutiveBaseBlockIt = BaseBlockByOffset.find(LastBlockEnd);
-				if (ConsecutiveBaseBlockIt != BaseBlockByOffset.end())
-				{
-					const BlockType& ConsecutiveBaseBlock = BaseBlocks[ConsecutiveBaseBlockIt->second];
-					if (ConsecutiveBaseBlock.HashStrong == NeedBlock.Hash)
-					{
-						UNSYNC_ASSERT(NeedBlock.Size == ConsecutiveBaseBlock.Size);
-						NeedBlock.SourceOffset = ConsecutiveBaseBlock.Offset;
-					}
-				}
-			}
-
-			NeedList.Base.push_back(NeedBlock);
-		}
-
-		NeedList.Sequence.push_back(ToHash128(SourceBlock.HashStrong));	 // #wip-widehash
-	}
-
-	return NeedList;
-}
-
-FNeedList
-DiffManifestBlocks(const FGenericBlockArray& SourceBlocks, const FGenericBlockArray& BaseBlocks)
-{
-	return DiffManifestBlocksT(SourceBlocks, BaseBlocks);
-}
-
-template<typename WeakHasher>
-FNeedList
-DiffBlocksParallelT(FIOReader&				  BaseDataReader,
-					uint32					  BlockSize,
-					EStrongHashAlgorithmID	  StrongHasher,
-					const FGenericBlockArray& SourceBlocks,
-					uint64					  BytesPerTask)
-{
-	auto TimeBegin = TimePointNow();
-
-	const uint64 BaseDataSize = BaseDataReader.GetSize();
-
-	THashSet<uint32, FIdentityHash32>							  SourceWeakHashSet;
-	THashSet<FGenericBlock, FBlockStrongHash, FBlockStrongHashEq> SourceStrongHashSet;
-
-	for (uint32 I = 0; I < uint32(SourceBlocks.size()); ++I)
-	{
-		SourceWeakHashSet.insert(SourceBlocks[I].HashWeak);
-		SourceStrongHashSet.insert(SourceBlocks[I]);
-	}
-
-	FNeedList NeedList;
-
-	struct FTask
-	{
-		uint64														  Offset = 0;
-		uint64														  Size	 = 0;
-		std::vector<FHash128>										  Sequence;
-		THashSet<FGenericBlock, FBlockStrongHash, FBlockStrongHashEq> BaseStrongHashSet;
-	};
-
-	BytesPerTask = std::max<uint64>(BlockSize, BytesPerTask);
-
-	std::vector<FTask> Tasks;
-	const uint64	   NumTasks = DivUp(BaseDataSize, BytesPerTask);
-	Tasks.resize(NumTasks);
-
-	FSchedulerSemaphore IoSemaphore(*GScheduler, 16);
-	FTaskGroup			TaskGroup = GScheduler->CreateTaskGroup(&IoSemaphore);
-
-	for (uint64 I = 0; I < NumTasks; ++I)
-	{
-		FTask& Task		 = Tasks[I];
-		uint64 TaskBegin = I * BytesPerTask;
-		uint64 TaskEnd	 = std::min(TaskBegin + BytesPerTask, BaseDataSize);
-
-		Task.Offset = TaskBegin;
-		Task.Size	= TaskEnd - TaskBegin;
-
-		auto ReadCallback =
-			[&SourceStrongHashSet, &SourceWeakHashSet, &Tasks, &TaskGroup, BaseDataSize, StrongHasher, BlockSize](
-				FIOBuffer CmdBuffer,
-				uint64	  CmdOffset,
-				uint64	  CmdReadSize,
-				uint64	  CmdUserData) {
-				TaskGroup.run([&SourceStrongHashSet,
-							   &SourceWeakHashSet,
-							   &Tasks,
-							   CmdBuffer = std::make_shared<FIOBuffer>(std::move(CmdBuffer)),
-							   CmdReadSize,
-							   CmdUserData,
-							   BaseDataSize,
-							   StrongHasher,
-							   BlockSize]() {
-					UNSYNC_ASSERT(CmdBuffer->GetSize() == CmdReadSize);
-
-					uint8* TaskBuffer = CmdBuffer->GetData();
-					uint64 TaskIndex  = CmdUserData;
-					FTask& Task		  = Tasks[TaskIndex];
-
-					UNSYNC_ASSERT(Task.Size == CmdReadSize);
-
-					const uint8* TaskEnd = TaskBuffer + Task.Size;
-
-					const uint32							  MaxWeakHashFalsePositives = 8;
-					THashMap<uint32, uint32, FIdentityHash32> WeakHashFalsePositives;
-					THashSet<uint32, FIdentityHash32>		  WeakHashBanList;
-
-					auto ScanFn = [&SourceWeakHashSet,
-								   &WeakHashBanList,
-								   BlockSize,
-								   &Task,
-								   TaskBuffer,
-								   StrongHasher,
-								   &SourceStrongHashSet,
-								   &WeakHashFalsePositives,
-								   TaskEnd,
-								   BaseDataSize](const uint8* WindowBegin, const uint8* WindowEnd, uint32 WindowHash) {
-						uint64 ThisBlockSize = WindowEnd - WindowBegin;
-
-						if (SourceWeakHashSet.find(WindowHash) != SourceWeakHashSet.end() &&
-							WeakHashBanList.find(WindowHash) == WeakHashBanList.end())
-						{
-							UNSYNC_ASSERT(ThisBlockSize <= BlockSize);
-
-							FGenericBlock BaseBlock;
-							BaseBlock.Offset	 = Task.Offset + (WindowBegin - TaskBuffer);
-							BaseBlock.Size		 = uint32(ThisBlockSize);
-							BaseBlock.HashWeak	 = WindowHash;
-							BaseBlock.HashStrong = ComputeHash(WindowBegin, ThisBlockSize, StrongHasher);
-
-							auto SourceBlockIt = SourceStrongHashSet.find(BaseBlock);
-							if (SourceBlockIt != SourceStrongHashSet.end())
-							{
-								const FGenericBlock& SourceBlock = *SourceBlockIt;
-
-								Task.BaseStrongHashSet.insert(BaseBlock);
-								Task.Sequence.push_back(SourceBlock.HashStrong.ToHash128());  // #wip-widehash
-
-								return true;
-							}
-
-							uint32 FalsePositives = WeakHashFalsePositives[WindowHash]++;
-							if (FalsePositives >= MaxWeakHashFalsePositives)
-							{
-								WeakHashBanList.insert(WindowHash);
-							}
-						}
-
-						return WindowEnd == TaskEnd && (Task.Offset + Task.Size) != BaseDataSize;
-					};
-
-					HashScan<WeakHasher>(TaskBuffer, Task.Size, BlockSize, ScanFn);
-				});
-			};
-
-		BaseDataReader.ReadAsync(Task.Offset, Task.Size, I, ReadCallback);
-	}
-
-	BaseDataReader.FlushAll();
-	TaskGroup.wait();
-
-	THashSet<FGenericBlock, FBlockStrongHash, FBlockStrongHashEq> BaseStrongHashSet;
-
-	for (FTask& Task : Tasks)
-	{
-		NeedList.Sequence.insert(NeedList.Sequence.end(), Task.Sequence.begin(), Task.Sequence.end());
-
-		for (const FGenericBlock& Block : Task.BaseStrongHashSet)
-		{
-			BaseStrongHashSet.insert(Block);
-		}
-	}
-
-	uint64 NeedBaseBytes   = 0;
-	uint64 NeedSourceBytes = 0;
-
-	for (const FGenericBlock& SourceBlock : SourceBlocks)
-	{
-		FNeedBlock NeedBlock;
-		NeedBlock.Size		   = SourceBlock.Size;
-		NeedBlock.TargetOffset = SourceBlock.Offset;
-		NeedBlock.Hash		   = SourceBlock.HashStrong;
-
-		auto BaseBlockIt = BaseStrongHashSet.find(SourceBlock);
-		if (BaseBlockIt != BaseStrongHashSet.end())
-		{
-			NeedBlock.SourceOffset = BaseBlockIt->Offset;
-			NeedList.Base.push_back(NeedBlock);
-			NeedBaseBytes += BaseBlockIt->Size;
-		}
-		else
-		{
-			NeedBlock.SourceOffset = SourceBlock.Offset;
-			NeedList.Source.push_back(NeedBlock);
-			NeedSourceBytes += SourceBlock.Size;
-		}
-	}
-
-	double Duration = DurationSec(TimeBegin, TimePointNow());
-	UNSYNC_VERBOSE(L"Done in %.3f sec (%.3f MB / sec)", Duration, SizeMb(double(BaseDataSize) / Duration));
-
-	return NeedList;
-}
-
-FNeedList
-DiffBlocksParallel(FIOReader&				 BaseDataReader,
-				   uint32					 BlockSize,
-				   EWeakHashAlgorithmID		 WeakHasher,
-				   EStrongHashAlgorithmID	 StrongHasher,
-				   const FGenericBlockArray& SourceBlocks,
-				   uint64					 BytesPerTask)
-{
-	switch (WeakHasher)
-	{
-		case EWeakHashAlgorithmID::Naive:
-			return DiffBlocksParallelT<FRollingChecksum>(BaseDataReader, BlockSize, StrongHasher, SourceBlocks, BytesPerTask);
-		case EWeakHashAlgorithmID::BuzHash:
-			return DiffBlocksParallelT<FBuzHash>(BaseDataReader, BlockSize, StrongHasher, SourceBlocks, BytesPerTask);
-		default:
-			UNSYNC_FATAL(L"Unexpected weak hash algorithm id");
-			return {};
-	}
-}
-
-FNeedList
-DiffBlocks(FIOReader&				 BaseDataReader,
-		   uint32					 BlockSize,
-		   EWeakHashAlgorithmID		 WeakHasher,
-		   EStrongHashAlgorithmID	 StrongHasher,
-		   const FGenericBlockArray& SourceBlocks)
-{
-	const uint64 BytesPerTask = 32_MB;	// <-- reasonably OK balance between accuracy and speed
-	// const uint64 bytes_per_task = base_data_size; // <-- run single-threaded
-	return DiffBlocksParallel(BaseDataReader, BlockSize, WeakHasher, StrongHasher, SourceBlocks, BytesPerTask);
-}
-
-FNeedList
-DiffBlocks(const uint8*				 BaseData,
-		   uint64					 BaseDataSize,
-		   uint32					 BlockSize,
-		   EWeakHashAlgorithmID		 WeakHasher,
-		   EStrongHashAlgorithmID	 StrongHasher,
-		   const FGenericBlockArray& SourceBlocks)
-{
-	FMemReader BaseReader(BaseData, BaseDataSize);
-	return DiffBlocks(BaseReader, BlockSize, WeakHasher, StrongHasher, SourceBlocks);
-}
-
-FNeedList
-DiffBlocksParallel(const uint8*				 BaseData,
-				   uint64					 BaseDataSize,
-				   uint32					 BlockSize,
-				   EWeakHashAlgorithmID		 WeakHasher,
-				   EStrongHashAlgorithmID	 StrongHasher,
-				   const FGenericBlockArray& SourceBlocks,
-				   uint64					 BytesPerTask)
-{
-	FMemReader BaseReader(BaseData, BaseDataSize);
-	return DiffBlocksParallel(BaseReader, BlockSize, WeakHasher, StrongHasher, SourceBlocks, BytesPerTask);
-}
-
 std::vector<FCopyCommand>
 OptimizeNeedList(const std::vector<FNeedBlock>& Input, uint64 MaxMergedBlockSize)
 {
@@ -1512,46 +567,6 @@ SyncFile(const FPath& SourceFilePath, const FPath& BaseFilePath, const FPath& Ta
 	return SyncFile(SourceFilePath, SourceBlocks, BaseFile, TargetFilePath, Options);
 }
 
-std::error_code
-CopyFileIfNewer(const FPath& Source, const FPath& Target)
-{
-	FFileAttributes SourceAttr = GetFileAttrib(Source);
-	FFileAttributes TargetAttr = GetFileAttrib(Target);
-	std::error_code Ec;
-	if (SourceAttr.Size != TargetAttr.Size || SourceAttr.Mtime != TargetAttr.Mtime)
-	{
-		FileCopyOverwrite(Source, Target, Ec);
-	}
-	return Ec;
-}
-
-static bool
-IsNonCaseSensitiveFileSystem(const FPath& ExistingPath)
-{
-	UNSYNC_ASSERTF(PathExists(ExistingPath), L"IsCaseSensitiveFileSystem must be called with a path that exists on disk");
-
-	// Assume file system is case-sensitive if all-upper and all-lower versions of the path exist and resolve to the same FS entry.
-	// This is not 100% robust due to symlinks, but is good enough for most practical purposes.
-
-	FPath PathUpper = StringToUpper(ExistingPath.wstring());
-	FPath PathLower = StringToLower(ExistingPath.wstring());
-
-	if (PathExists(PathUpper) && PathExists(PathLower))
-	{
-		return std::filesystem::equivalent(ExistingPath, PathUpper) && std::filesystem::equivalent(PathLower, PathUpper);
-	}
-	else
-	{
-		return false;
-	}
-}
-
-static bool
-IsCaseSensitiveFileSystem(const FPath& ExistingPath)
-{
-	return !IsNonCaseSensitiveFileSystem(ExistingPath);
-}
-
 struct FPendingFileRename
 {
 	std::wstring Old;
@@ -3109,102 +2124,6 @@ FSyncFilter::Resolve(const FPath& Filename) const
 	return Result;
 }
 
-FBuffer
-GeneratePatch(const uint8*			 BaseData,
-			  uint64				 BaseDataSize,
-			  const uint8*			 SourceData,
-			  uint64				 SourceDataSize,
-			  uint32				 BlockSize,
-			  EWeakHashAlgorithmID	 WeakHasher,
-			  EStrongHashAlgorithmID StrongHasher,
-			  int32					 CompressionLevel)
-{
-	FBuffer Result;
-
-	FAlgorithmOptions Algorithm;
-	Algorithm.ChunkingAlgorithmId	= EChunkingAlgorithmID::FixedBlocks;
-	Algorithm.WeakHashAlgorithmId	= WeakHasher;
-	Algorithm.StrongHashAlgorithmId = StrongHasher;
-
-	UNSYNC_VERBOSE(L"Computing blocks for source (%.2f MB)", SizeMb(SourceDataSize));
-	FGenericBlockArray SourceBlocks = ComputeBlocks(SourceData, SourceDataSize, BlockSize, Algorithm);
-
-	FGenericBlockArray SourceValidation, BaseValidation;
-
-	{
-		FLogVerbosityScope VerbosityScope(false);
-		SourceValidation = ComputeBlocks(SourceData, SourceDataSize, FPatchHeader::VALIDATION_BLOCK_SIZE, Algorithm);
-		BaseValidation	 = ComputeBlocks(BaseData, BaseDataSize, FPatchHeader::VALIDATION_BLOCK_SIZE, Algorithm);
-	}
-
-	UNSYNC_VERBOSE(L"Computing difference for base (%.2f MB)", SizeMb(BaseDataSize));
-	FNeedList NeedList = DiffBlocks(BaseData, BaseDataSize, BlockSize, WeakHasher, StrongHasher, SourceBlocks);
-
-	if (IsSynchronized(NeedList, SourceBlocks))
-	{
-		return Result;
-	}
-
-	FPatchCommandList PatchCommands;
-	PatchCommands.Source = OptimizeNeedList(NeedList.Source);
-	PatchCommands.Base	 = OptimizeNeedList(NeedList.Base);
-
-	uint64 NeedFromSource = ComputeSize(NeedList.Source);
-	uint64 NeedFromBase	  = ComputeSize(NeedList.Base);
-	UNSYNC_VERBOSE(L"Need from source %.2f MB, from base: %.2f MB", SizeMb(NeedFromSource), SizeMb(NeedFromBase));
-
-	FVectorStreamOut Stream(Result);
-
-	FPatchHeader Header;
-	Header.SourceSize				 = SourceDataSize;
-	Header.BaseSize					 = BaseDataSize;
-	Header.NumSourceValidationBlocks = SourceValidation.size();
-	Header.NumBaseValidationBlocks	 = BaseValidation.size();
-	Header.NumSourceBlocks			 = PatchCommands.Source.size();
-	Header.NumBaseBlocks			 = PatchCommands.Base.size();
-	Header.BlockSize				 = BlockSize;
-	Header.WeakHashAlgorithmId		 = WeakHasher;
-	Header.StrongHashAlgorithmId	 = StrongHasher;
-	Stream.Write(&Header, sizeof(Header));
-
-	FHash128 HeaderHash = HashBlake3Bytes<FHash128>(Result.Data(), Result.Size());
-	Stream.Write(&HeaderHash, sizeof(HeaderHash));
-
-	for (const FGenericBlock& Block : SourceValidation)
-	{
-		Stream.Write(&Block, sizeof(Block));
-	}
-	for (const FGenericBlock& Block : BaseValidation)
-	{
-		Stream.Write(&Block, sizeof(Block));
-	}
-	for (FCopyCommand& Cmd : PatchCommands.Source)
-	{
-		Stream.Write(&Cmd, sizeof(Cmd));
-	}
-	for (FCopyCommand& Cmd : PatchCommands.Base)
-	{
-		Stream.Write(&Cmd, sizeof(Cmd));
-	}
-
-	FHash128 BlockHash = HashBlake3Bytes<FHash128>(Result.Data(), Result.Size());
-	Stream.Write(&BlockHash, sizeof(BlockHash));
-
-	for (const FCopyCommand& Cmd : PatchCommands.Source)
-	{
-		Stream.Write(SourceData + Cmd.SourceOffset, Cmd.Size);
-	}
-
-	const uint64 RawPatchSize = Result.Size();
-	UNSYNC_VERBOSE(L"Compressing patch (%.2f MB raw)", SizeMb(RawPatchSize));
-
-	Result = Compress(Result.Data(), Result.Size(), CompressionLevel);
-
-	UNSYNC_VERBOSE(L"Compressed patch size: %.2f MB", SizeMb(Result.Size()));
-
-	return Result;
-}
-
 FNeedListSize
 ComputeNeedListSize(const FNeedList& NeedList)
 {
diff --git a/Engine/Source/Programs/Unsync/Private/UnsyncCore.h b/Engine/Source/Programs/Unsync/Private/UnsyncCore.h
--- a/Engine/Source/Programs/Unsync/Private/UnsyncCore.h
+++ b/Engine/Source/Programs/Unsync/Private/UnsyncCore.h
@@ -125,8 +125,6 @@ struct FNeedListSize
 
 FNeedListSize ComputeNeedListSize(const FNeedList& NeedList);
 
-const std::string& GetVersionString();
-
 struct FBlockSourceInfo
 {
 	// FPath  FilePath; // TODO: pass this in to allow filtering by name, etc.
@@ -152,72 +150,9 @@ struct FComputeBlocksParams
 	bool bAllowThreading = true;
 };
 
-struct FComputeBlocksResult
-{
-	FGenericBlockArray Blocks;
-	FGenericBlockArray MacroBlocks;
-};
-
-FComputeBlocksResult ComputeBlocks(FIOReader& Reader, const FComputeBlocksParams& Params);
-FComputeBlocksResult ComputeBlocks(const uint8* Data, uint64 Size, const FComputeBlocksParams& Params);
-FComputeBlocksResult ComputeBlocksVariable(FIOReader& Reader, const FComputeBlocksParams& Params);
-
-FGenericBlockArray ComputeBlocks(FIOReader& Reader, uint32 BlockSize, FAlgorithmOptions Algorithm);
-FGenericBlockArray ComputeBlocks(const uint8* Data, uint64 Size, uint32 BlockSize, FAlgorithmOptions Algorithm);
-FGenericBlockArray ComputeBlocksVariable(FIOReader&				Reader,
-										 uint32					BlockSize,
-										 EWeakHashAlgorithmID	WeakHasher,
-										 EStrongHashAlgorithmID StrongHasher);
-
-
-FNeedList DiffBlocks(const uint8*			   BaseData,
-					 uint64					   BaseDataSize,
-					 uint32					   BlockSize,
-					 EWeakHashAlgorithmID	   WeakHasher,
-					 EStrongHashAlgorithmID	   StrongHasher,
-					 const FGenericBlockArray& SourceBlocks);
-
-FNeedList DiffBlocksParallel(const uint8*			   BaseData,
-							 uint64					   BaseDataSize,
-							 uint32					   BlockSize,
-							 EWeakHashAlgorithmID	   WeakHasher,
-							 EStrongHashAlgorithmID	   StrongHasher,
-							 const FGenericBlockArray& SourceBlocks,
-							 uint64					   BytesPerTask);
-
-FNeedList DiffBlocks(FIOReader&				   BaseDataReader,
-					 uint32					   BlockSize,
-					 EWeakHashAlgorithmID	   WeakHasher,
-					 EStrongHashAlgorithmID	   StrongHasher,
-					 const FGenericBlockArray& SourceBlocks);
-
-FNeedList DiffBlocksParallel(FIOReader&				   BaseDataReader,
-							 uint32					   BlockSize,
-							 EWeakHashAlgorithmID	   WeakHasher,
-							 EStrongHashAlgorithmID	   StrongHasher,
-							 const FGenericBlockArray& SourceBlocks,
-							 uint64					   BytesPerTask);
-
-FNeedList DiffBlocksVariable(FIOReader&				   BaseDataReader,
-							 uint32					   BlockSize,
-							 EWeakHashAlgorithmID	   WeakHasher,
-							 EStrongHashAlgorithmID	   StrongHasher,
-							 const FGenericBlockArray& SourceBlocks);
-
-FNeedList DiffManifestBlocks(const FGenericBlockArray& SourceBlocks, const FGenericBlockArray& BaseBlocks);
 
 std::vector<FCopyCommand> OptimizeNeedList(const std::vector<FNeedBlock>& Input, uint64 MaxMergedBlockSize = 8_MB);
 
-
-FBuffer GeneratePatch(const uint8*			 BaseData,
-					  uint64				 BaseDataSize,
-					  const uint8*			 SourceData,
-					  uint64				 SourceDataSize,
-					  uint32				 BlockSize,
-					  EWeakHashAlgorithmID	 WeakHasher,
-					  EStrongHashAlgorithmID StrongHasher,
-					  int32					 CompressionLevel = 3);
-
 bool IsSynchronized(const FNeedList& NeedList, const FGenericBlockArray& SourceBlocks);
 
 enum class EFileSyncStatus
@@ -352,4 +287,23 @@ struct FCmdInfoOptions
 };
 int32 CmdInfo(const FCmdInfoOptions& Options);
 
+template<typename BlockType>
+bool
+ValidateBlockListT(const std::vector<BlockType>& Blocks)
+{
+	uint64 CurrentOffset = 0;
+	for (const BlockType& Block : Blocks)
+	{
+		if (CurrentOffset != Block.Offset)
+		{
+			UNSYNC_ERROR(L"Found block at unexpected offset. Blocks are expected to be ordered by offset and contiguous.");
+			return false;
+		}
+
+		CurrentOffset += Block.Size;
+	}
+
+	return true;
+}
+
 }  // namespace unsync
diff --git a/Engine/Source/Programs/Unsync/Private/UnsyncDiff.cpp b/Engine/Source/Programs/Unsync/Private/UnsyncDiff.cpp
new file mode 100644
--- /dev/null
+++ b/Engine/Source/Programs/Unsync/Private/UnsyncDiff.cpp
@@ -0,0 +1,453 @@
+// Copyright Epic Games, Inc. All Rights Reserved.
+
+#include "UnsyncDiff.h"
+#include "UnsyncChunking.h"
+#include "UnsyncCompression.h"
+#include "UnsyncFile.h"
+#include "UnsyncScan.h"
+#include "UnsyncScheduler.h"
+
+namespace unsync {
+
+FNeedList
+DiffBlocksVariable(FIOReader&				 BaseDataReader,
+				   uint32					 BlockSize,
+				   EWeakHashAlgorithmID		 WeakHasher,
+				   EStrongHashAlgorithmID	 StrongHasher,
+				   const FGenericBlockArray& SourceBlocks)
+{
+	FGenericBlockArray BaseBlocks = ComputeBlocksVariable(BaseDataReader, BlockSize, WeakHasher, StrongHasher);
+	if (!ValidateBlockListT(BaseBlocks))
+	{
+		UNSYNC_FATAL(L"Base block list validation failed");
+	}
+
+	return DiffManifestBlocks(SourceBlocks, BaseBlocks);
+}
+
+template<typename BlockType>
+FNeedList
+DiffManifestBlocksT(const std::vector<BlockType>& SourceBlocks, const std::vector<BlockType>& BaseBlocks)
+{
+	FNeedList NeedList;
+
+	struct BlockIndexAndCount
+	{
+		uint64 Index = 0;
+		uint64 Count = 0;
+	};
+
+	THashMap<typename BlockType::StrongHashType, BlockIndexAndCount> BaseBlockMap;
+	THashMap<uint64, uint64>										 BaseBlockByOffset;
+
+	for (uint64 I = 0; I < BaseBlocks.size(); ++I)
+	{
+		const BlockType& Block			= BaseBlocks[I];
+		BaseBlockByOffset[Block.Offset] = I;
+
+		auto Existing = BaseBlockMap.find(Block.HashStrong);
+		if (Existing == BaseBlockMap.end())
+		{
+			BlockIndexAndCount Item;
+			Item.Index = I;
+			Item.Count = 1;
+			BaseBlockMap.insert(std::make_pair(Block.HashStrong, Item));
+		}
+		else
+		{
+			Existing->second.Count += 1;
+		}
+	}
+
+	for (uint64 I = 0; I < SourceBlocks.size(); ++I)
+	{
+		const BlockType& SourceBlock = SourceBlocks[I];
+
+		auto BaseBlockIt = BaseBlockMap.find(SourceBlock.HashStrong);
+		if (BaseBlockIt == BaseBlockMap.end())
+		{
+			FNeedBlock NeedBlock;
+			NeedBlock.Hash		   = SourceBlock.HashStrong;
+			NeedBlock.Size		   = SourceBlock.Size;
+			NeedBlock.SourceOffset = SourceBlock.Offset;
+			NeedBlock.TargetOffset = SourceBlock.Offset;
+			NeedList.Source.push_back(NeedBlock);
+		}
+		else
+		{
+			BlockIndexAndCount IndexAndCount = BaseBlockIt->second;
+			const BlockType&   BaseBlock	 = BaseBlocks[IndexAndCount.Index];
+
+			UNSYNC_ASSERT(BaseBlock.Size == SourceBlock.Size);
+
+			FNeedBlock NeedBlock;
+			NeedBlock.Hash		   = BaseBlock.HashStrong;
+			NeedBlock.Size		   = BaseBlock.Size;
+			NeedBlock.SourceOffset = BaseBlock.Offset;
+			NeedBlock.TargetOffset = SourceBlock.Offset;
+
+			const FNeedBlock* LastBaseNeedBlock = NeedList.Base.empty() ? nullptr : &(NeedList.Base.back());
+
+			// Try to preserve contiguous base data reads
+			if (LastBaseNeedBlock)
+			{
+				uint64 LastBlockEnd			  = LastBaseNeedBlock->SourceOffset + LastBaseNeedBlock->Size;
+				auto   ConsecutiveBaseBlockIt = BaseBlockByOffset.find(LastBlockEnd);
+				if (ConsecutiveBaseBlockIt != BaseBlockByOffset.end())
+				{
+					const BlockType& ConsecutiveBaseBlock = BaseBlocks[ConsecutiveBaseBlockIt->second];
+					if (ConsecutiveBaseBlock.HashStrong == NeedBlock.Hash)
+					{
+						UNSYNC_ASSERT(NeedBlock.Size == ConsecutiveBaseBlock.Size);
+						NeedBlock.SourceOffset = ConsecutiveBaseBlock.Offset;
+					}
+				}
+			}
+
+			NeedList.Base.push_back(NeedBlock);
+		}
+
+		NeedList.Sequence.push_back(ToHash128(SourceBlock.HashStrong));	 // #wip-widehash
+	}
+
+	return NeedList;
+}
+
+FNeedList
+DiffManifestBlocks(const FGenericBlockArray& SourceBlocks, const FGenericBlockArray& BaseBlocks)
+{
+	return DiffManifestBlocksT(SourceBlocks, BaseBlocks);
+}
+
+template<typename WeakHasher>
+FNeedList
+DiffBlocksParallelT(FIOReader&				  BaseDataReader,
+					uint32					  BlockSize,
+					EStrongHashAlgorithmID	  StrongHasher,
+					const FGenericBlockArray& SourceBlocks,
+					uint64					  BytesPerTask)
+{
+	auto TimeBegin = TimePointNow();
+
+	const uint64 BaseDataSize = BaseDataReader.GetSize();
+
+	THashSet<uint32, FIdentityHash32>							  SourceWeakHashSet;
+	THashSet<FGenericBlock, FBlockStrongHash, FBlockStrongHashEq> SourceStrongHashSet;
+
+	for (uint32 I = 0; I < uint32(SourceBlocks.size()); ++I)
+	{
+		SourceWeakHashSet.insert(SourceBlocks[I].HashWeak);
+		SourceStrongHashSet.insert(SourceBlocks[I]);
+	}
+
+	FNeedList NeedList;
+
+	struct FTask
+	{
+		uint64														  Offset = 0;
+		uint64														  Size	 = 0;
+		std::vector<FHash128>										  Sequence;
+		THashSet<FGenericBlock, FBlockStrongHash, FBlockStrongHashEq> BaseStrongHashSet;
+	};
+
+	BytesPerTask = std::max<uint64>(BlockSize, BytesPerTask);
+
+	std::vector<FTask> Tasks;
+	const uint64	   NumTasks = DivUp(BaseDataSize, BytesPerTask);
+	Tasks.resize(NumTasks);
+
+	FSchedulerSemaphore IoSemaphore(*GScheduler, 16);
+	FTaskGroup			TaskGroup = GScheduler->CreateTaskGroup(&IoSemaphore);
+
+	for (uint64 I = 0; I < NumTasks; ++I)
+	{
+		FTask& Task		 = Tasks[I];
+		uint64 TaskBegin = I * BytesPerTask;
+		uint64 TaskEnd	 = std::min(TaskBegin + BytesPerTask, BaseDataSize);
+
+		Task.Offset = TaskBegin;
+		Task.Size	= TaskEnd - TaskBegin;
+
+		auto ReadCallback =
+			[&SourceStrongHashSet, &SourceWeakHashSet, &Tasks, &TaskGroup, BaseDataSize, StrongHasher, BlockSize](FIOBuffer CmdBuffer,
+																												  uint64	CmdOffset,
+																												  uint64	CmdReadSize,
+																												  uint64	CmdUserData)
+		{
+			TaskGroup.run(
+				[&SourceStrongHashSet,
+				 &SourceWeakHashSet,
+				 &Tasks,
+				 CmdBuffer = std::make_shared<FIOBuffer>(std::move(CmdBuffer)),
+				 CmdReadSize,
+				 CmdUserData,
+				 BaseDataSize,
+				 StrongHasher,
+				 BlockSize]()
+				{
+					UNSYNC_ASSERT(CmdBuffer->GetSize() == CmdReadSize);
+
+					uint8* TaskBuffer = CmdBuffer->GetData();
+					uint64 TaskIndex  = CmdUserData;
+					FTask& Task		  = Tasks[TaskIndex];
+
+					UNSYNC_ASSERT(Task.Size == CmdReadSize);
+
+					const uint8* TaskEnd = TaskBuffer + Task.Size;
+
+					const uint32							  MaxWeakHashFalsePositives = 8;
+					THashMap<uint32, uint32, FIdentityHash32> WeakHashFalsePositives;
+					THashSet<uint32, FIdentityHash32>		  WeakHashBanList;
+
+					auto ScanFn = [&SourceWeakHashSet,
+								   &WeakHashBanList,
+								   BlockSize,
+								   &Task,
+								   TaskBuffer,
+								   StrongHasher,
+								   &SourceStrongHashSet,
+								   &WeakHashFalsePositives,
+								   TaskEnd,
+								   BaseDataSize](const uint8* WindowBegin, const uint8* WindowEnd, uint32 WindowHash)
+					{
+						uint64 ThisBlockSize = WindowEnd - WindowBegin;
+
+						if (SourceWeakHashSet.find(WindowHash) != SourceWeakHashSet.end() &&
+							WeakHashBanList.find(WindowHash) == WeakHashBanList.end())
+						{
+							UNSYNC_ASSERT(ThisBlockSize <= BlockSize);
+
+							FGenericBlock BaseBlock;
+							BaseBlock.Offset	 = Task.Offset + (WindowBegin - TaskBuffer);
+							BaseBlock.Size		 = uint32(ThisBlockSize);
+							BaseBlock.HashWeak	 = WindowHash;
+							BaseBlock.HashStrong = ComputeHash(WindowBegin, ThisBlockSize, StrongHasher);
+
+							auto SourceBlockIt = SourceStrongHashSet.find(BaseBlock);
+							if (SourceBlockIt != SourceStrongHashSet.end())
+							{
+								const FGenericBlock& SourceBlock = *SourceBlockIt;
+
+								Task.BaseStrongHashSet.insert(BaseBlock);
+								Task.Sequence.push_back(SourceBlock.HashStrong.ToHash128());  // #wip-widehash
+
+								return true;
+							}
+
+							uint32 FalsePositives = WeakHashFalsePositives[WindowHash]++;
+							if (FalsePositives >= MaxWeakHashFalsePositives)
+							{
+								WeakHashBanList.insert(WindowHash);
+							}
+						}
+
+						return WindowEnd == TaskEnd && (Task.Offset + Task.Size) != BaseDataSize;
+					};
+
+					HashScan<WeakHasher>(TaskBuffer, Task.Size, BlockSize, ScanFn);
+				});
+		};
+
+		BaseDataReader.ReadAsync(Task.Offset, Task.Size, I, ReadCallback);
+	}
+
+	BaseDataReader.FlushAll();
+	TaskGroup.wait();
+
+	THashSet<FGenericBlock, FBlockStrongHash, FBlockStrongHashEq> BaseStrongHashSet;
+
+	for (FTask& Task : Tasks)
+	{
+		NeedList.Sequence.insert(NeedList.Sequence.end(), Task.Sequence.begin(), Task.Sequence.end());
+
+		for (const FGenericBlock& Block : Task.BaseStrongHashSet)
+		{
+			BaseStrongHashSet.insert(Block);
+		}
+	}
+
+	uint64 NeedBaseBytes   = 0;
+	uint64 NeedSourceBytes = 0;
+
+	for (const FGenericBlock& SourceBlock : SourceBlocks)
+	{
+		FNeedBlock NeedBlock;
+		NeedBlock.Size		   = SourceBlock.Size;
+		NeedBlock.TargetOffset = SourceBlock.Offset;
+		NeedBlock.Hash		   = SourceBlock.HashStrong;
+
+		auto BaseBlockIt = BaseStrongHashSet.find(SourceBlock);
+		if (BaseBlockIt != BaseStrongHashSet.end())
+		{
+			NeedBlock.SourceOffset = BaseBlockIt->Offset;
+			NeedList.Base.push_back(NeedBlock);
+			NeedBaseBytes += BaseBlockIt->Size;
+		}
+		else
+		{
+			NeedBlock.SourceOffset = SourceBlock.Offset;
+			NeedList.Source.push_back(NeedBlock);
+			NeedSourceBytes += SourceBlock.Size;
+		}
+	}
+
+	double Duration = DurationSec(TimeBegin, TimePointNow());
+	UNSYNC_VERBOSE(L"Done in %.3f sec (%.3f MB / sec)", Duration, SizeMb(double(BaseDataSize) / Duration));
+
+	return NeedList;
+}
+
+FNeedList
+DiffBlocksParallel(FIOReader&				 BaseDataReader,
+				   uint32					 BlockSize,
+				   EWeakHashAlgorithmID		 WeakHasher,
+				   EStrongHashAlgorithmID	 StrongHasher,
+				   const FGenericBlockArray& SourceBlocks,
+				   uint64					 BytesPerTask)
+{
+	switch (WeakHasher)
+	{
+		case EWeakHashAlgorithmID::Naive:
+			return DiffBlocksParallelT<FRollingChecksum>(BaseDataReader, BlockSize, StrongHasher, SourceBlocks, BytesPerTask);
+		case EWeakHashAlgorithmID::BuzHash:
+			return DiffBlocksParallelT<FBuzHash>(BaseDataReader, BlockSize, StrongHasher, SourceBlocks, BytesPerTask);
+		default:
+			UNSYNC_FATAL(L"Unexpected weak hash algorithm id");
+			return {};
+	}
+}
+
+FNeedList
+DiffBlocks(FIOReader&				 BaseDataReader,
+		   uint32					 BlockSize,
+		   EWeakHashAlgorithmID		 WeakHasher,
+		   EStrongHashAlgorithmID	 StrongHasher,
+		   const FGenericBlockArray& SourceBlocks)
+{
+	const uint64 BytesPerTask = 32_MB;	// <-- reasonably OK balance between accuracy and speed
+	// const uint64 bytes_per_task = base_data_size; // <-- run single-threaded
+	return DiffBlocksParallel(BaseDataReader, BlockSize, WeakHasher, StrongHasher, SourceBlocks, BytesPerTask);
+}
+
+FNeedList
+DiffBlocks(const uint8*				 BaseData,
+		   uint64					 BaseDataSize,
+		   uint32					 BlockSize,
+		   EWeakHashAlgorithmID		 WeakHasher,
+		   EStrongHashAlgorithmID	 StrongHasher,
+		   const FGenericBlockArray& SourceBlocks)
+{
+	FMemReader BaseReader(BaseData, BaseDataSize);
+	return DiffBlocks(BaseReader, BlockSize, WeakHasher, StrongHasher, SourceBlocks);
+}
+
+FNeedList
+DiffBlocksParallel(const uint8*				 BaseData,
+				   uint64					 BaseDataSize,
+				   uint32					 BlockSize,
+				   EWeakHashAlgorithmID		 WeakHasher,
+				   EStrongHashAlgorithmID	 StrongHasher,
+				   const FGenericBlockArray& SourceBlocks,
+				   uint64					 BytesPerTask)
+{
+	FMemReader BaseReader(BaseData, BaseDataSize);
+	return DiffBlocksParallel(BaseReader, BlockSize, WeakHasher, StrongHasher, SourceBlocks, BytesPerTask);
+}
+
+FBuffer
+GeneratePatch(const uint8*			 BaseData,
+			  uint64				 BaseDataSize,
+			  const uint8*			 SourceData,
+			  uint64				 SourceDataSize,
+			  uint32				 BlockSize,
+			  EWeakHashAlgorithmID	 WeakHasher,
+			  EStrongHashAlgorithmID StrongHasher,
+			  int32					 CompressionLevel)
+{
+	FBuffer Result;
+
+	FAlgorithmOptions Algorithm;
+	Algorithm.ChunkingAlgorithmId	= EChunkingAlgorithmID::FixedBlocks;
+	Algorithm.WeakHashAlgorithmId	= WeakHasher;
+	Algorithm.StrongHashAlgorithmId = StrongHasher;
+
+	UNSYNC_VERBOSE(L"Computing blocks for source (%.2f MB)", SizeMb(SourceDataSize));
+	FGenericBlockArray SourceBlocks = ComputeBlocks(SourceData, SourceDataSize, BlockSize, Algorithm);
+
+	FGenericBlockArray SourceValidation, BaseValidation;
+
+	{
+		FLogVerbosityScope VerbosityScope(false);
+		SourceValidation = ComputeBlocks(SourceData, SourceDataSize, FPatchHeader::VALIDATION_BLOCK_SIZE, Algorithm);
+		BaseValidation	 = ComputeBlocks(BaseData, BaseDataSize, FPatchHeader::VALIDATION_BLOCK_SIZE, Algorithm);
+	}
+
+	UNSYNC_VERBOSE(L"Computing difference for base (%.2f MB)", SizeMb(BaseDataSize));
+	FNeedList NeedList = DiffBlocks(BaseData, BaseDataSize, BlockSize, WeakHasher, StrongHasher, SourceBlocks);
+
+	if (IsSynchronized(NeedList, SourceBlocks))
+	{
+		return Result;
+	}
+
+	FPatchCommandList PatchCommands;
+	PatchCommands.Source = OptimizeNeedList(NeedList.Source);
+	PatchCommands.Base	 = OptimizeNeedList(NeedList.Base);
+
+	uint64 NeedFromSource = ComputeSize(NeedList.Source);
+	uint64 NeedFromBase	  = ComputeSize(NeedList.Base);
+	UNSYNC_VERBOSE(L"Need from source %.2f MB, from base: %.2f MB", SizeMb(NeedFromSource), SizeMb(NeedFromBase));
+
+	FVectorStreamOut Stream(Result);
+
+	FPatchHeader Header;
+	Header.SourceSize				 = SourceDataSize;
+	Header.BaseSize					 = BaseDataSize;
+	Header.NumSourceValidationBlocks = SourceValidation.size();
+	Header.NumBaseValidationBlocks	 = BaseValidation.size();
+	Header.NumSourceBlocks			 = PatchCommands.Source.size();
+	Header.NumBaseBlocks			 = PatchCommands.Base.size();
+	Header.BlockSize				 = BlockSize;
+	Header.WeakHashAlgorithmId		 = WeakHasher;
+	Header.StrongHashAlgorithmId	 = StrongHasher;
+	Stream.Write(&Header, sizeof(Header));
+
+	FHash128 HeaderHash = HashBlake3Bytes<FHash128>(Result.Data(), Result.Size());
+	Stream.Write(&HeaderHash, sizeof(HeaderHash));
+
+	for (const FGenericBlock& Block : SourceValidation)
+	{
+		Stream.Write(&Block, sizeof(Block));
+	}
+	for (const FGenericBlock& Block : BaseValidation)
+	{
+		Stream.Write(&Block, sizeof(Block));
+	}
+	for (FCopyCommand& Cmd : PatchCommands.Source)
+	{
+		Stream.Write(&Cmd, sizeof(Cmd));
+	}
+	for (FCopyCommand& Cmd : PatchCommands.Base)
+	{
+		Stream.Write(&Cmd, sizeof(Cmd));
+	}
+
+	FHash128 BlockHash = HashBlake3Bytes<FHash128>(Result.Data(), Result.Size());
+	Stream.Write(&BlockHash, sizeof(BlockHash));
+
+	for (const FCopyCommand& Cmd : PatchCommands.Source)
+	{
+		Stream.Write(SourceData + Cmd.SourceOffset, Cmd.Size);
+	}
+
+	const uint64 RawPatchSize = Result.Size();
+	UNSYNC_VERBOSE(L"Compressing patch (%.2f MB raw)", SizeMb(RawPatchSize));
+
+	Result = Compress(Result.Data(), Result.Size(), CompressionLevel);
+
+	UNSYNC_VERBOSE(L"Compressed patch size: %.2f MB", SizeMb(Result.Size()));
+
+	return Result;
+}
+
+}  // namespace unsync
diff --git a/Engine/Source/Programs/Unsync/Private/UnsyncDiff.h b/Engine/Source/Programs/Unsync/Private/UnsyncDiff.h
new file mode 100644
--- /dev/null
+++ b/Engine/Source/Programs/Unsync/Private/UnsyncDiff.h
@@ -0,0 +1,53 @@
+// Copyright Epic Games, Inc. All Rights Reserved.
+
+#include "UnsyncCommon.h"
+#include "UnsyncCore.h"
+
+namespace unsync {
+
+FNeedList DiffBlocks(const uint8*			   BaseData,
+					 uint64					   BaseDataSize,
+					 uint32					   BlockSize,
+					 EWeakHashAlgorithmID	   WeakHasher,
+					 EStrongHashAlgorithmID	   StrongHasher,
+					 const FGenericBlockArray& SourceBlocks);
+
+FNeedList DiffBlocksParallel(const uint8*			   BaseData,
+							 uint64					   BaseDataSize,
+							 uint32					   BlockSize,
+							 EWeakHashAlgorithmID	   WeakHasher,
+							 EStrongHashAlgorithmID	   StrongHasher,
+							 const FGenericBlockArray& SourceBlocks,
+							 uint64					   BytesPerTask);
+
+FNeedList DiffBlocks(FIOReader&				   BaseDataReader,
+					 uint32					   BlockSize,
+					 EWeakHashAlgorithmID	   WeakHasher,
+					 EStrongHashAlgorithmID	   StrongHasher,
+					 const FGenericBlockArray& SourceBlocks);
+
+FNeedList DiffBlocksParallel(FIOReader&				   BaseDataReader,
+							 uint32					   BlockSize,
+							 EWeakHashAlgorithmID	   WeakHasher,
+							 EStrongHashAlgorithmID	   StrongHasher,
+							 const FGenericBlockArray& SourceBlocks,
+							 uint64					   BytesPerTask);
+
+FNeedList DiffBlocksVariable(FIOReader&				   BaseDataReader,
+							 uint32					   BlockSize,
+							 EWeakHashAlgorithmID	   WeakHasher,
+							 EStrongHashAlgorithmID	   StrongHasher,
+							 const FGenericBlockArray& SourceBlocks);
+
+FNeedList DiffManifestBlocks(const FGenericBlockArray& SourceBlocks, const FGenericBlockArray& BaseBlocks);
+
+FBuffer GeneratePatch(const uint8*			 BaseData,
+					  uint64				 BaseDataSize,
+					  const uint8*			 SourceData,
+					  uint64				 SourceDataSize,
+					  uint32				 BlockSize,
+					  EWeakHashAlgorithmID	 WeakHasher,
+					  EStrongHashAlgorithmID StrongHasher,
+					  int32					 CompressionLevel = 3);
+
+}  // namespace unsync
diff --git a/Engine/Source/Programs/Unsync/Private/UnsyncFile.cpp b/Engine/Source/Programs/Unsync/Private/UnsyncFile.cpp
--- a/Engine/Source/Programs/Unsync/Private/UnsyncFile.cpp
+++ b/Engine/Source/Programs/Unsync/Private/UnsyncFile.cpp
@@ -4,7 +4,7 @@
 #include "UnsyncCore.h"
 #include "UnsyncMemory.h"
 #include "UnsyncThread.h"
-#include "UnsyncThread.h"
+#include "UnsyncScheduler.h"
 
 #include <mutex>
 
@@ -129,6 +129,47 @@ GetRelativePath(const FPath& Path, const FPath& Base)
 	return {};
 }
 
+std::error_code
+CopyFileIfNewer(const FPath& Source, const FPath& Target)
+{
+	FFileAttributes SourceAttr = GetFileAttrib(Source);
+	FFileAttributes TargetAttr = GetFileAttrib(Target);
+	std::error_code Ec;
+	if (SourceAttr.Size != TargetAttr.Size || SourceAttr.Mtime != TargetAttr.Mtime)
+	{
+		FileCopyOverwrite(Source, Target, Ec);
+	}
+	return Ec;
+}
+
+
+bool
+IsNonCaseSensitiveFileSystem(const FPath& ExistingPath)
+{
+	UNSYNC_ASSERTF(PathExists(ExistingPath), L"IsCaseSensitiveFileSystem must be called with a path that exists on disk");
+
+	// Assume file system is case-sensitive if all-upper and all-lower versions of the path exist and resolve to the same FS entry.
+	// This is not 100% robust due to symlinks, but is good enough for most practical purposes.
+
+	FPath PathUpper = StringToUpper(ExistingPath.wstring());
+	FPath PathLower = StringToLower(ExistingPath.wstring());
+
+	if (PathExists(PathUpper) && PathExists(PathLower))
+	{
+		return std::filesystem::equivalent(ExistingPath, PathUpper) && std::filesystem::equivalent(PathLower, PathUpper);
+	}
+	else
+	{
+		return false;
+	}
+}
+
+bool
+IsCaseSensitiveFileSystem(const FPath& ExistingPath)
+{
+	return !IsNonCaseSensitiveFileSystem(ExistingPath);
+}
+
 FFileAttributes GetCachedFileAttrib(const FPath& Path, FFileAttributeCache& AttribCache)
 {
 	FFileAttributes Result;
@@ -1399,6 +1440,54 @@ TestFileTime()
 	}
 }
 
+uint64
+BlockingReadLarge(FIOReader& Reader, uint64 Offset, uint64 Size, uint8* OutputBuffer, uint64 OutputBufferSize)
+{
+	const uint64 BytesPerRead = 2_MB;
+	const uint64 ReadEnd	  = std::min(Offset + Size, Reader.GetSize());
+	const uint64 ClampedSize  = ReadEnd - Offset;
+
+	std::atomic<uint64> TotalReadSize = 0;
+
+	if (ClampedSize == 0)
+	{
+		return TotalReadSize;
+	}
+
+	FSchedulerSemaphore IoSemaphore(*GScheduler, 16);
+	FTaskGroup			CopyTasks = GScheduler->CreateTaskGroup(&IoSemaphore);
+
+	uint64 NumReads = DivUp(ClampedSize, BytesPerRead);
+	for (uint64 ReadIndex = 0; ReadIndex < NumReads; ++ReadIndex)
+	{
+		const uint64 ThisBatchSize	= CalcChunkSize(ReadIndex, BytesPerRead, ClampedSize);
+		const uint64 OutputOffset	= BytesPerRead * ReadIndex;
+		const uint64 ThisReadOffset = Offset + OutputOffset;
+
+		auto ReadCallback = [OutputBuffer, OutputBufferSize, &TotalReadSize, &CopyTasks](FIOBuffer CmdBuffer,
+																						 uint64	   CmdSourceOffset,
+																						 uint64	   CmdReadSize,
+																						 uint64	   OutputOffset)
+		{
+			UNSYNC_ASSERT(OutputOffset + CmdReadSize <= OutputBufferSize);
+
+			CopyTasks.run(
+				[OutputBuffer, OutputOffset, CmdReadSize, CmdBuffer = MakeShared(std::move(CmdBuffer)), &TotalReadSize]()
+				{
+					memcpy(OutputBuffer + OutputOffset, CmdBuffer->GetData(), CmdReadSize);
+					TotalReadSize += CmdReadSize;
+				});
+		};
+
+		Reader.ReadAsync(ThisReadOffset, ThisBatchSize, OutputOffset, ReadCallback);
+	}
+
+	Reader.FlushAll();
+	CopyTasks.wait();
+
+	return TotalReadSize;
+}
+
 void
 TestFileAttrib()
 {
diff --git a/Engine/Source/Programs/Unsync/Private/UnsyncFile.h b/Engine/Source/Programs/Unsync/Private/UnsyncFile.h
--- a/Engine/Source/Programs/Unsync/Private/UnsyncFile.h
+++ b/Engine/Source/Programs/Unsync/Private/UnsyncFile.h
@@ -458,6 +458,8 @@ IsReadOnly(std::filesystem::perms Perms)
 	return (Perms & std::filesystem::perms::owner_write) == std::filesystem::perms::none;
 }
 
+uint64 BlockingReadLarge(FIOReader& Reader, uint64 Offset, uint64 Size, uint8* OutputBuffer, uint64 OutputBufferSize);
+
 FFileAttributes GetFileAttrib(const FPath& Path, FFileAttributeCache* AttribCache = nullptr);
 FFileAttributes GetCachedFileAttrib(const FPath& Path, FFileAttributeCache& AttribCache);
 
@@ -473,6 +475,9 @@ bool			FileCopy(const FPath& From, const FPath& To, std::error_code& OutErrorCod
 bool			FileCopyOverwrite(const FPath& From, const FPath& To, std::error_code& OutErrorCode);
 bool			FileRemove(const FPath& Path, std::error_code& OutErrorCode);
 FPath			GetRelativePath(const FPath& Path, const FPath& Base);
+std::error_code CopyFileIfNewer(const FPath& Source, const FPath& Target);
+bool			IsNonCaseSensitiveFileSystem(const FPath& ExistingPath);
+bool			IsCaseSensitiveFileSystem(const FPath& ExistingPath);
 
 // Returns number of bytes that can be written to the given path.
 // Returns ~0ull if the available space could not be determined.
diff --git a/Engine/Source/Programs/Unsync/Private/UnsyncHash.h b/Engine/Source/Programs/Unsync/Private/UnsyncHash.h
--- a/Engine/Source/Programs/Unsync/Private/UnsyncHash.h
+++ b/Engine/Source/Programs/Unsync/Private/UnsyncHash.h
@@ -42,8 +42,7 @@ using FHash128 = THashValue<16>;
 using FHash160 = THashValue<20>;
 using FHash256 = THashValue<32>;
 
-enum class EHashType : uint8
-{
+enum class EHashType : uint8 {
 	Unknown,
 	Unknown_128,
 	Unknown_160,
@@ -307,6 +306,18 @@ FGenericHash ComputeHash(const uint8* Data, uint64 Size, EStrongHashAlgorithmID
 FGenericHash ComputeHash(const FBuffer& Buffer, EStrongHashAlgorithmID Algorithm);
 FGenericHash ComputeHash(const FBufferView& Buffer, EStrongHashAlgorithmID Algorithm);
 
+inline FHash128
+ToHash128(const FHash128& X)
+{
+	return X;
+}
+
+inline FHash128
+ToHash128(const FGenericHash& X)
+{
+	return X.ToHash128();
+}
+
 }  // namespace unsync
 
 namespace std {
diff --git a/Engine/Source/Programs/Unsync/Private/UnsyncHttp.cpp b/Engine/Source/Programs/Unsync/Private/UnsyncHttp.cpp
--- a/Engine/Source/Programs/Unsync/Private/UnsyncHttp.cpp
+++ b/Engine/Source/Programs/Unsync/Private/UnsyncHttp.cpp
@@ -5,6 +5,7 @@
 #include "UnsyncUtil.h"
 #include "UnsyncRemote.h"
 #include "UnsyncAuth.h"
+#include "UnsyncVersion.h"
 
 #include <http_parser.h>
 #include <string.h>
diff --git a/Engine/Source/Programs/Unsync/Private/UnsyncLog.cpp b/Engine/Source/Programs/Unsync/Private/UnsyncLog.cpp
--- a/Engine/Source/Programs/Unsync/Private/UnsyncLog.cpp
+++ b/Engine/Source/Programs/Unsync/Private/UnsyncLog.cpp
@@ -5,6 +5,7 @@
 #include "UnsyncError.h"
 #include "UnsyncFile.h"
 #include "UnsyncUtil.h"
+#include "UnsyncVersion.h"
 
 #include <stdarg.h>
 #include <stdio.h>
diff --git a/Engine/Source/Programs/Unsync/Private/UnsyncMain.cpp b/Engine/Source/Programs/Unsync/Private/UnsyncMain.cpp
--- a/Engine/Source/Programs/Unsync/Private/UnsyncMain.cpp
+++ b/Engine/Source/Programs/Unsync/Private/UnsyncMain.cpp
@@ -18,6 +18,7 @@
 #include "UnsyncThread.h"
 #include "UnsyncUtil.h"
 #include "UnsyncScheduler.h"
+#include "UnsyncVersion.h"
 
 UNSYNC_THIRD_PARTY_INCLUDES_START
 #if UNSYNC_PLATFORM_WINDOWS
diff --git a/Engine/Source/Programs/Unsync/Private/UnsyncManifest.cpp b/Engine/Source/Programs/Unsync/Private/UnsyncManifest.cpp
--- a/Engine/Source/Programs/Unsync/Private/UnsyncManifest.cpp
+++ b/Engine/Source/Programs/Unsync/Private/UnsyncManifest.cpp
@@ -8,6 +8,7 @@
 #include "UnsyncSerialization.h"
 #include "UnsyncThread.h"
 #include "UnsyncScheduler.h"
+#include "UnsyncChunking.h"
 
 UNSYNC_THIRD_PARTY_INCLUDES_START
 #include <blake3.h>
diff --git a/Engine/Source/Programs/Unsync/Private/UnsyncProtocol.h b/Engine/Source/Programs/Unsync/Private/UnsyncProtocol.h
--- a/Engine/Source/Programs/Unsync/Private/UnsyncProtocol.h
+++ b/Engine/Source/Programs/Unsync/Private/UnsyncProtocol.h
@@ -18,6 +18,8 @@
 
 namespace unsync {
 
+static constexpr uint32 MAX_BLOCK_SIZE = uint32(1_MB);
+
 enum class EChunkingAlgorithmID : uint64
 {
 	Invalid		   = 0,
diff --git a/Engine/Source/Programs/Unsync/Private/UnsyncSerialization.cpp b/Engine/Source/Programs/Unsync/Private/UnsyncSerialization.cpp
--- a/Engine/Source/Programs/Unsync/Private/UnsyncSerialization.cpp
+++ b/Engine/Source/Programs/Unsync/Private/UnsyncSerialization.cpp
@@ -2,6 +2,7 @@
 
 #include "UnsyncSerialization.h"
 #include "UnsyncFile.h"
+#include "UnsyncVersion.h"
 
 namespace unsync {
 
diff --git a/Engine/Source/Programs/Unsync/Private/UnsyncTarget.cpp b/Engine/Source/Programs/Unsync/Private/UnsyncTarget.cpp
--- a/Engine/Source/Programs/Unsync/Private/UnsyncTarget.cpp
+++ b/Engine/Source/Programs/Unsync/Private/UnsyncTarget.cpp
@@ -10,6 +10,7 @@
 #include "UnsyncScavenger.h"
 #include "UnsyncThread.h"
 #include "UnsyncScheduler.h"
+#include "UnsyncChunking.h"
 
 namespace unsync {
 
diff --git a/Engine/Source/Programs/Unsync/Private/UnsyncTest.cpp b/Engine/Source/Programs/Unsync/Private/UnsyncTest.cpp
--- a/Engine/Source/Programs/Unsync/Private/UnsyncTest.cpp
+++ b/Engine/Source/Programs/Unsync/Private/UnsyncTest.cpp
@@ -1,14 +1,16 @@
 // Copyright Epic Games, Inc. All Rights Reserved.
 
 #include "UnsyncTest.h"
+#include "UnsyncChunking.h"
 #include "UnsyncCore.h"
+#include "UnsyncDiff.h"
 #include "UnsyncFile.h"
 #include "UnsyncHash.h"
 #include "UnsyncScan.h"
+#include "UnsyncTarget.h"
 #include "UnsyncTest.h"
 #include "UnsyncThread.h"
 #include "UnsyncUtil.h"
-#include "UnsyncTarget.h"
 
 UNSYNC_THIRD_PARTY_INCLUDES_START
 #include <md5-sse2.h>
diff --git a/Engine/Source/Programs/Unsync/Private/UnsyncVersion.cpp b/Engine/Source/Programs/Unsync/Private/UnsyncVersion.cpp
new file mode 100644
--- /dev/null
+++ b/Engine/Source/Programs/Unsync/Private/UnsyncVersion.cpp
@@ -0,0 +1,38 @@
+// Copyright Epic Games, Inc. All Rights Reserved.
+
+#include "UnsyncVersion.h"
+
+namespace unsync {
+
+const std::string&
+GetVersionString()
+{
+	static std::string Result = []()
+	{
+		// TODO: generate a version string based on git state
+		const char* GitRev	  = "";
+		const char* GitBranch = "";
+		// const char* GIT_TAG = nullptr;
+
+		static char Str[256];
+
+		if (strlen(GitBranch) && strlen(GitRev))
+		{
+			snprintf(Str, sizeof(Str), UNSYNC_VERSION_STR " [%s:%s]", GitBranch, GitRev);
+		}
+		else if (strlen(GitRev))
+		{
+			snprintf(Str, sizeof(Str), UNSYNC_VERSION_STR " [%s]", GitRev);
+		}
+		else
+		{
+			snprintf(Str, sizeof(Str), UNSYNC_VERSION_STR);
+		}
+
+		return std::string(Str);
+	}();
+
+	return Result;
+}
+
+}  // namespace unsync
\ No newline at end of file
diff --git a/Engine/Source/Programs/Unsync/Private/UnsyncVersion.h b/Engine/Source/Programs/Unsync/Private/UnsyncVersion.h
new file mode 100644
--- /dev/null
+++ b/Engine/Source/Programs/Unsync/Private/UnsyncVersion.h
@@ -0,0 +1,12 @@
+// Copyright Epic Games, Inc. All Rights Reserved.
+
+#include "UnsyncCommon.h"
+
+#define UNSYNC_VERSION_STR "1.0.78"
+
+namespace unsync {
+
+const std::string&
+GetVersionString();
+
+}
