unsync - Further split utility code into dedicated modules

diff --git a/Engine/Source/Programs/Unsync/CMakeLists.txt b/Engine/Source/Programs/Unsync/CMakeLists.txt
--- a/Engine/Source/Programs/Unsync/CMakeLists.txt
+++ b/Engine/Source/Programs/Unsync/CMakeLists.txt
@@ -64,6 +64,8 @@ set(src
 	Private/UnsyncError.h
 	Private/UnsyncFile.cpp
 	Private/UnsyncFile.h
+	Private/UnsyncFilter.cpp
+	Private/UnsyncFilter.h
 	Private/UnsyncHash.cpp
 	Private/UnsyncHash.h
 	Private/UnsyncHashTable.h
@@ -102,6 +104,8 @@ set(src
 	Private/UnsyncSerialization.h
 	Private/UnsyncSocket.cpp
 	Private/UnsyncSocket.h
+	Private/UnsyncSource.cpp
+	Private/UnsyncSource.h
 	Private/UnsyncTarget.cpp
 	Private/UnsyncTarget.h
 	Private/UnsyncTest.cpp
diff --git a/Engine/Source/Programs/Unsync/Private/UnsyncChunking.h b/Engine/Source/Programs/Unsync/Private/UnsyncChunking.h
--- a/Engine/Source/Programs/Unsync/Private/UnsyncChunking.h
+++ b/Engine/Source/Programs/Unsync/Private/UnsyncChunking.h
@@ -1,5 +1,7 @@
 // Copyright Epic Games, Inc. All Rights Reserved.
 
+#pragma once
+
 #include "UnsyncCommon.h"
 #include "UnsyncCore.h"
 
diff --git a/Engine/Source/Programs/Unsync/Private/UnsyncCmdQuery.h b/Engine/Source/Programs/Unsync/Private/UnsyncCmdQuery.h
--- a/Engine/Source/Programs/Unsync/Private/UnsyncCmdQuery.h
+++ b/Engine/Source/Programs/Unsync/Private/UnsyncCmdQuery.h
@@ -1,5 +1,7 @@
 // Copyright Epic Games, Inc. All Rights Reserved.
 
+#pragma once
+
 #include "UnsyncCommon.h"
 #include "UnsyncRemote.h"
 
diff --git a/Engine/Source/Programs/Unsync/Private/UnsyncCmdSync.cpp b/Engine/Source/Programs/Unsync/Private/UnsyncCmdSync.cpp
--- a/Engine/Source/Programs/Unsync/Private/UnsyncCmdSync.cpp
+++ b/Engine/Source/Programs/Unsync/Private/UnsyncCmdSync.cpp
@@ -3,6 +3,7 @@
 #include "UnsyncCmdSync.h"
 #include "UnsyncFile.h"
 #include "UnsyncProxy.h"
+#include "UnsyncFilter.h"
 
 namespace unsync {
 
@@ -155,15 +156,15 @@ CmdSync(const FCmdSyncOptions& Options)
 
 			if (bSourceFileSystemRequired)
 			{
-				SyncOptions.SourceType = ESyncSourceType::FileSystem;
+				SyncOptions.SourceType = ESourceType::FileSystem;
 			}
 			else if (bSourceIsManifestHash)
 			{
-				SyncOptions.SourceType = ESyncSourceType::ServerWithManifestHash;
+				SyncOptions.SourceType = ESourceType::ServerWithManifestHash;
 			}
 			else
 			{
-				SyncOptions.SourceType = ESyncSourceType::Server;
+				SyncOptions.SourceType = ESourceType::Server;
 			}
 
 			SyncOptions.Source					   = ResolvedSource;
diff --git a/Engine/Source/Programs/Unsync/Private/UnsyncCore.cpp b/Engine/Source/Programs/Unsync/Private/UnsyncCore.cpp
--- a/Engine/Source/Programs/Unsync/Private/UnsyncCore.cpp
+++ b/Engine/Source/Programs/Unsync/Private/UnsyncCore.cpp
@@ -19,6 +19,8 @@
 #include "UnsyncThread.h"
 #include "UnsyncUtil.h"
 #include "UnsyncVersion.h"
+#include "UnsyncFilter.h"
+#include "UnsyncSource.h"
 
 #include <condition_variable>
 #include <filesystem>
@@ -715,62 +717,6 @@ FixFileNameCases(const FPath& RootPath, const std::vector<FPendingFileRename>& P
 	return true;
 }
 
-static bool
-MergeManifests(FDirectoryManifest& Existing, const FDirectoryManifest& Other, bool bCaseSensitive)
-{
-	if (!Existing.IsValid())
-	{
-		Existing = Other;
-		return true;
-	}
-
-	if (!AlgorithmOptionsCompatible(Existing.Algorithm, Other.Algorithm))
-	{
-		UNSYNC_ERROR("Trying to merge incompatible manifests (diff algorithm options do not match)");
-		return false;
-	}
-
-	if (bCaseSensitive)
-	{
-		// Trivial case: just replace existing entries
-		for (const auto& OtherFile : Other.Files)
-		{
-			Existing.Files[OtherFile.first] = OtherFile.second;
-		}
-	}
-	else
-	{
-		// Lookup table of lowercase -> original file name used to replace conflicting entries on non-case-sensitive filesystems
-		// TODO: Could potentially add case-sensitive/insensitive entry lookup helper functions to FDirectoryManifest itself in the future
-		std::unordered_map<std::wstring, std::wstring> ExistingFileNamesLowerCase;
-
-		for (auto& ExistingEntry : Existing.Files)
-		{
-			std::wstring FileNameLowerCase = StringToLower(ExistingEntry.first);
-			ExistingFileNamesLowerCase.insert(std::pair<std::wstring, std::wstring>(FileNameLowerCase, ExistingEntry.first));
-		}
-
-		for (const auto& OtherFile : Other.Files)
-		{
-			std::wstring OtherNameLowerCase = StringToLower(OtherFile.first);
-			auto		 LowerCaseEntry		= ExistingFileNamesLowerCase.find(OtherNameLowerCase);
-			if (LowerCaseEntry != ExistingFileNamesLowerCase.end())
-			{
-				// Remove file with conflicting case and add entry from the other manifest instead
-				const std::wstring& ExistingNameOriginalCase = LowerCaseEntry->second;
-				Existing.Files.erase(ExistingNameOriginalCase);
-
-				// Update the lookup table entry to refer to the name we're about to insert
-				ExistingFileNamesLowerCase[LowerCaseEntry->first] = OtherFile.first;
-			}
-
-			Existing.Files[OtherFile.first] = OtherFile.second;
-		}
-	}
-
-	return true;
-}
-
 // Delete files from target directory that are not in the source directory manifest
 static void
 DeleteUnnecessaryFiles(const FPath&				 TargetDirectory,
@@ -1201,9 +1147,9 @@ SyncDirectory(const FSyncDirectoryOptions& SyncOptions)
 
 	FTimePoint TimeBegin = TimePointNow();
 
-	const bool bFileSystemSource = SyncOptions.SourceType == ESyncSourceType::FileSystem;
+	const bool bFileSystemSource = SyncOptions.SourceType == ESourceType::FileSystem;
 	const bool bServerSource =
-		SyncOptions.SourceType == ESyncSourceType::Server || SyncOptions.SourceType == ESyncSourceType::ServerWithManifestHash;
+		SyncOptions.SourceType == ESourceType::Server || SyncOptions.SourceType == ESourceType::ServerWithManifestHash;
 
 	UNSYNC_ASSERT(bFileSystemSource || bServerSource);
 
@@ -1270,7 +1216,7 @@ SyncDirectory(const FSyncDirectoryOptions& SyncOptions)
 
 	FTimingLogger ManifestLoadTimingLogger("Manifest load time", ELogLevel::Info);
 
-	if (SyncOptions.SourceType == ESyncSourceType::ServerWithManifestHash)
+	if (SyncOptions.SourceType == ESourceType::ServerWithManifestHash)
 	{
 		if (!ProxyPool.IsValid())
 		{
@@ -2098,32 +2044,6 @@ SyncDirectory(const FSyncDirectoryOptions& SyncOptions)
 	return bSyncSucceeded;
 }
 
-FPath
-FSyncFilter::Resolve(const FPath& Filename) const
-{
-	std::wstring FilenameLower = StringToLower(Filename.wstring());
-
-	FPath Result;
-	for (const auto& Alias : DfsAliases)
-	{
-		// TODO: add a case-insensitive find() helper
-		size_t Pos = FilenameLower.find(StringToLower(Alias.Source.wstring()));
-		if (Pos == 0)
-		{
-			auto Tail = (Filename.wstring().substr(Alias.Source.wstring().length() + 1));
-			Result	  = Alias.Target / Tail;
-			break;
-		}
-	}
-
-	if (Result.empty())
-	{
-		Result = Filename;
-	}
-
-	return Result;
-}
-
 FNeedListSize
 ComputeNeedListSize(const FNeedList& NeedList)
 {
@@ -2143,107 +2063,6 @@ ComputeNeedListSize(const FNeedList& NeedList)
 	return Result;
 }
 
-static void
-AddCommaSeparatedWordsToList(const std::wstring& CommaSeparatedWords, std::vector<std::wstring>& Output)
-{
-	size_t Offset = 0;
-	size_t Len	  = CommaSeparatedWords.length();
-	while (Offset < Len)
-	{
-		size_t MatchOffset = CommaSeparatedWords.find(L',', Offset);
-		if (MatchOffset == std::wstring::npos)
-		{
-			MatchOffset = Len;
-		}
-
-		std::wstring Word = CommaSeparatedWords.substr(Offset, MatchOffset - Offset);
-		Output.push_back(Word);
-
-		Offset = MatchOffset + 1;
-	}
-}
-
-void
-FSyncFilter::IncludeInSync(const std::wstring& CommaSeparatedWords)
-{
-	AddCommaSeparatedWordsToList(CommaSeparatedWords, SyncIncludedWords);
-}
-
-void
-FSyncFilter::ExcludeFromSync(const std::wstring& CommaSeparatedWords)
-{
-	AddCommaSeparatedWordsToList(CommaSeparatedWords, SyncExcludedWords);
-}
-
-void
-FSyncFilter::ExcludeFromCleanup(const std::wstring& CommaSeparatedWords)
-{
-	AddCommaSeparatedWordsToList(CommaSeparatedWords, CleanupExcludedWords);
-}
-
-bool
-FSyncFilter::ShouldSync(const FPath& Filename) const
-{
-#if UNSYNC_PLATFORM_WINDOWS
-	return ShouldSync(Filename.native());
-#else
-	return ShouldSync(Filename.wstring());
-#endif
-}
-
-bool
-FSyncFilter::ShouldSync(const std::wstring& Filename) const
-{
-	bool bInclude = SyncIncludedWords.empty();	// Include everything if there are no specific inclusions
-	for (const std::wstring& Word : SyncIncludedWords)
-	{
-		if (Filename.find(Word) != std::wstring::npos)
-		{
-			bInclude = true;
-			break;
-		}
-	}
-
-	if (!bInclude)
-	{
-		return false;
-	}
-
-	for (const std::wstring& Word : SyncExcludedWords)
-	{
-		if (Filename.find(Word) != std::wstring::npos)
-		{
-			return false;
-		}
-	}
-
-	return true;
-}
-
-bool
-FSyncFilter::ShouldCleanup(const FPath& Filename) const
-{
-#if UNSYNC_PLATFORM_WINDOWS
-	return ShouldCleanup(Filename.native());
-#else
-	return ShouldCleanup(Filename.wstring());
-#endif
-}
-
-bool
-FSyncFilter::ShouldCleanup(const std::wstring& Filename) const
-{
-	for (const std::wstring& Word : CleanupExcludedWords)
-	{
-		if (Filename.find(Word) != std::wstring::npos)
-		{
-			return false;
-		}
-	}
-
-	return true;
-}
-
 const wchar_t*
 ToString(EFileSyncStatus Status)
 {
@@ -2268,199 +2087,4 @@ ToString(EFileSyncStatus Status)
 	}
 }
 
-THashMap<FGenericHash, FGenericBlock>
-BuildBlockMap(const FDirectoryManifest& Manifest, bool bNeedMacroBlocks)
-{
-	THashMap<FGenericHash, FGenericBlock> Result;
-	for (const auto& It : Manifest.Files)
-	{
-		const FFileManifest& File = It.second;
-		if (bNeedMacroBlocks)
-		{
-			for (const FGenericBlock& Block : File.MacroBlocks)
-			{
-				Result[Block.HashStrong] = Block;
-			}
-		}
-		else
-		{
-			for (const FGenericBlock& Block : File.Blocks)
-			{
-				Result[Block.HashStrong] = Block;
-			}
-		}
-	}
-	return Result;
-}
-
-void
-LogManifestDiff(ELogLevel LogLevel, const FDirectoryManifest& ManifestA, const FDirectoryManifest& ManifestB)
-{
-	THashMap<FGenericHash, FGenericBlock> BlocksA = BuildBlockMap(ManifestA, false);
-	THashMap<FGenericHash, FGenericBlock> BlocksB = BuildBlockMap(ManifestB, false);
-
-	THashMap<FGenericHash, FGenericBlock> MacroBlocksA = BuildBlockMap(ManifestA, true);
-	THashMap<FGenericHash, FGenericBlock> MacroBlocksB = BuildBlockMap(ManifestB, true);
-
-	uint32 NumCommonBlocks		= 0;
-	uint64 TotalCommonBlockSize = 0;
-	uint64 TotalSizeA			= 0;
-	uint64 TotalSizeB			= 0;
-
-	uint64 PatchSizeFromAtoB = 0;
-
-	for (const auto& ItA : BlocksA)
-	{
-		TotalSizeA += ItA.second.Size;
-		auto ItB = BlocksB.find(ItA.first);
-		if (ItB != BlocksB.end())
-		{
-			NumCommonBlocks++;
-			TotalCommonBlockSize += ItA.second.Size;
-		}
-	}
-
-	for (const auto& ItB : BlocksB)
-	{
-		TotalSizeB += ItB.second.Size;
-		if (BlocksA.find(ItB.first) == BlocksA.end())
-		{
-			PatchSizeFromAtoB += ItB.second.Size;
-		}
-	}
-
-	uint32 NumCommonMacroBlocks		 = 0;
-	uint64 TotalCommonMacroBlockSize = 0;
-	for (const auto& ItA : MacroBlocksA)
-	{
-		auto ItB = MacroBlocksB.find(ItA.first);
-		if (ItB != MacroBlocksB.end())
-		{
-			NumCommonMacroBlocks++;
-			TotalCommonMacroBlockSize += ItA.second.Size;
-		}
-	}
-
-	LogPrintf(LogLevel, L"Common macro blocks: %d, %.3f MB\n", NumCommonMacroBlocks, SizeMb(TotalCommonMacroBlockSize));
-
-	LogPrintf(LogLevel,
-			  L"Common blocks: %d, %.3f MB (%.2f%% of A, %.2f%% of B)\n",
-			  NumCommonBlocks,
-			  SizeMb(TotalCommonBlockSize),
-			  100.0 * double(TotalCommonBlockSize) / double(TotalSizeA),
-			  100.0 * double(TotalCommonBlockSize) / double(TotalSizeB));
-
-	LogPrintf(LogLevel, L"Patch size: %.3f MB\n", SizeMb(PatchSizeFromAtoB));
-}
-
-void
-FilterManifest(const FSyncFilter& SyncFilter, FDirectoryManifest& Manifest)
-{
-	auto It = Manifest.Files.begin();
-	while (It != Manifest.Files.end())
-	{
-		if (SyncFilter.ShouldSync(It->first))
-		{
-			++It;
-		}
-		else
-		{
-			It = Manifest.Files.erase(It);
-		}
-	}
-}
-
-int32
-CmdInfo(const FCmdInfoOptions& Options)
-{
-	FPath DirectoryManifestPathA = IsDirectory(Options.InputA) ? (Options.InputA / ".unsync" / "manifest.bin") : Options.InputA;
-	FPath DirectoryManifestPathB = IsDirectory(Options.InputB) ? (Options.InputB / ".unsync" / "manifest.bin") : Options.InputB;
-
-	FDirectoryManifest ManifestA;
-
-	bool bManifestAValid = LoadDirectoryManifest(ManifestA, Options.InputA, DirectoryManifestPathA);
-
-	if (!bManifestAValid)
-	{
-		return 1;
-	}
-
-	LogPrintf(ELogLevel::Info, L"Manifest A: %ls\n", DirectoryManifestPathA.wstring().c_str());
-
-	if (Options.SyncFilter)
-	{
-		FilterManifest(*Options.SyncFilter, ManifestA);
-	}
-
-	{
-		UNSYNC_LOG_INDENT;
-		LogManifestInfo(ELogLevel::Info, ManifestA);
-	}
-
-	if (Options.bListFiles)
-	{
-		UNSYNC_LOG_INDENT;
-		LogManifestFiles(ELogLevel::Info, ManifestA);
-	}
-
-	if (Options.InputB.empty())
-	{
-		return 0;
-	}
-
-	LogPrintf(ELogLevel::Info, L"\n");
-
-	FDirectoryManifest ManifestB;
-
-	bool bManifestBValid = LoadDirectoryManifest(ManifestB, Options.InputB, DirectoryManifestPathB);
-
-	if (!bManifestBValid)
-	{
-		return 1;
-	}
-
-	LogPrintf(ELogLevel::Info, L"Manifest B: %ls\n", DirectoryManifestPathB.wstring().c_str());
-
-	if (Options.SyncFilter)
-	{
-		FilterManifest(*Options.SyncFilter, ManifestB);
-	}
-
-	{
-		UNSYNC_LOG_INDENT;
-		LogManifestInfo(ELogLevel::Info, ManifestB);
-	}
-	if (Options.bListFiles)
-	{
-		UNSYNC_LOG_INDENT;
-		LogManifestFiles(ELogLevel::Info, ManifestB);
-	}
-
-	LogPrintf(ELogLevel::Info, L"\n");
-	LogPrintf(ELogLevel::Info, L"Difference:\n");
-
-	{
-		UNSYNC_LOG_INDENT;
-		LogManifestDiff(ELogLevel::Info, ManifestA, ManifestB);
-	}
-
-	return 0;
-}
-
-FHash256
-ComputeSerializedManifestHash(const FDirectoryManifest& Manifest)
-{
-	FBuffer			 ManifestBuffer;
-	FVectorStreamOut ManifestStream(ManifestBuffer);
-	bool			 bSerializedOk = SaveDirectoryManifest(Manifest, ManifestStream);
-	UNSYNC_ASSERT(bSerializedOk);
-	return HashBlake3Bytes<FHash256>(ManifestBuffer.Data(), ManifestBuffer.Size());
-}
-
-FHash160
-ComputeSerializedManifestHash160(const FDirectoryManifest& Manifest)
-{
-	return ToHash160(ComputeSerializedManifestHash(Manifest));
-}
-
 }  // namespace unsync
diff --git a/Engine/Source/Programs/Unsync/Private/UnsyncCore.h b/Engine/Source/Programs/Unsync/Private/UnsyncCore.h
--- a/Engine/Source/Programs/Unsync/Private/UnsyncCore.h
+++ b/Engine/Source/Programs/Unsync/Private/UnsyncCore.h
@@ -31,6 +31,9 @@ struct FIOReader;
 struct FIOWriter;
 struct FIOReaderWriter;
 struct FPackWriteContext;
+struct FSyncFilter;
+
+enum class ESourceType : uint8;
 
 struct FIdentityHash32
 {
@@ -220,41 +223,9 @@ FFileSyncResult SyncFile(const FPath&			 SourceFilePath,
 						 const FPath&			 TargetFilePath,
 						 const FSyncFileOptions& Options);
 
-struct FSyncFilter
-{
-	FSyncFilter() = default;
-
-	// By default all files will be included, calling this will include only files containing these substrings
-	void IncludeInSync(const std::wstring& CommaSeparatedWords);
-	void ExcludeFromSync(const std::wstring& CommaSeparatedWords);
-	void ExcludeFromCleanup(const std::wstring& CommaSeparatedWords);
-
-	bool  ShouldSync(const FPath& Filename) const;
-	bool  ShouldSync(const std::wstring& Filename) const;
-
-	bool  ShouldCleanup(const FPath& Filename) const;
-	bool  ShouldCleanup(const std::wstring& Filename) const;
-
-	FPath Resolve(const FPath& Filename) const;
-
-	std::vector<std::wstring> SyncIncludedWords;
-	std::vector<std::wstring> SyncExcludedWords; // any paths that contain these words will not be synced
-	std::vector<std::wstring> CleanupExcludedWords; // any paths that contain these words will not be deleted after sync
-
-	std::vector<FDfsAlias>	  DfsAliases;
-};
-
-enum class ESyncSourceType
-{
-	Unknown,
-	FileSystem,
-	Server,
-	ServerWithManifestHash,
-};
-
 struct FSyncDirectoryOptions
 {
-	ESyncSourceType	   SourceType;
+	ESourceType		   SourceType;
 	FPath			   Source;			   // remote data location
 	FPath			   Target;			   // output target location
 	FPath			   Base;			   // base data location, which typically is the same as sync target
@@ -278,15 +249,6 @@ bool SyncDirectory(const FSyncDirectoryOptions& SyncOptions);
 FBlock128			   ToBlock128(const FGenericBlock& GenericBlock);
 std::vector<FBlock128> ToBlock128(FGenericBlockArray& GenericBlocks);
 
-struct FCmdInfoOptions
-{
-	FPath InputA;
-	FPath InputB;
-	bool bListFiles = false;
-	const FSyncFilter* SyncFilter = nullptr;
-};
-int32 CmdInfo(const FCmdInfoOptions& Options);
-
 template<typename BlockType>
 bool
 ValidateBlockListT(const std::vector<BlockType>& Blocks)
diff --git a/Engine/Source/Programs/Unsync/Private/UnsyncDiff.h b/Engine/Source/Programs/Unsync/Private/UnsyncDiff.h
--- a/Engine/Source/Programs/Unsync/Private/UnsyncDiff.h
+++ b/Engine/Source/Programs/Unsync/Private/UnsyncDiff.h
@@ -1,5 +1,7 @@
 // Copyright Epic Games, Inc. All Rights Reserved.
 
+#pragma once
+
 #include "UnsyncCommon.h"
 #include "UnsyncCore.h"
 
diff --git a/Engine/Source/Programs/Unsync/Private/UnsyncFile.cpp b/Engine/Source/Programs/Unsync/Private/UnsyncFile.cpp
--- a/Engine/Source/Programs/Unsync/Private/UnsyncFile.cpp
+++ b/Engine/Source/Programs/Unsync/Private/UnsyncFile.cpp
@@ -5,6 +5,7 @@
 #include "UnsyncMemory.h"
 #include "UnsyncThread.h"
 #include "UnsyncScheduler.h"
+#include "UnsyncFilter.h"
 
 #include <mutex>
 
diff --git a/Engine/Source/Programs/Unsync/Private/UnsyncFilter.cpp b/Engine/Source/Programs/Unsync/Private/UnsyncFilter.cpp
new file mode 100644
--- /dev/null
+++ b/Engine/Source/Programs/Unsync/Private/UnsyncFilter.cpp
@@ -0,0 +1,135 @@
+// Copyright Epic Games, Inc. All Rights Reserved.
+
+#include "UnsyncFilter.h"
+#include "UnsyncUtil.h"
+
+namespace unsync {
+
+static void
+AddCommaSeparatedWordsToList(const std::wstring& CommaSeparatedWords, std::vector<std::wstring>& Output)
+{
+	size_t Offset = 0;
+	size_t Len	  = CommaSeparatedWords.length();
+	while (Offset < Len)
+	{
+		size_t MatchOffset = CommaSeparatedWords.find(L',', Offset);
+		if (MatchOffset == std::wstring::npos)
+		{
+			MatchOffset = Len;
+		}
+
+		std::wstring Word = CommaSeparatedWords.substr(Offset, MatchOffset - Offset);
+		Output.push_back(Word);
+
+		Offset = MatchOffset + 1;
+	}
+}
+
+void
+FSyncFilter::IncludeInSync(const std::wstring& CommaSeparatedWords)
+{
+	AddCommaSeparatedWordsToList(CommaSeparatedWords, SyncIncludedWords);
+}
+
+void
+FSyncFilter::ExcludeFromSync(const std::wstring& CommaSeparatedWords)
+{
+	AddCommaSeparatedWordsToList(CommaSeparatedWords, SyncExcludedWords);
+}
+
+void
+FSyncFilter::ExcludeFromCleanup(const std::wstring& CommaSeparatedWords)
+{
+	AddCommaSeparatedWordsToList(CommaSeparatedWords, CleanupExcludedWords);
+}
+
+bool
+FSyncFilter::ShouldSync(const FPath& Filename) const
+{
+#if UNSYNC_PLATFORM_WINDOWS
+	return ShouldSync(Filename.native());
+#else
+	return ShouldSync(Filename.wstring());
+#endif
+}
+
+bool
+FSyncFilter::ShouldSync(const std::wstring& Filename) const
+{
+	bool bInclude = SyncIncludedWords.empty();	// Include everything if there are no specific inclusions
+	for (const std::wstring& Word : SyncIncludedWords)
+	{
+		if (Filename.find(Word) != std::wstring::npos)
+		{
+			bInclude = true;
+			break;
+		}
+	}
+
+	if (!bInclude)
+	{
+		return false;
+	}
+
+	for (const std::wstring& Word : SyncExcludedWords)
+	{
+		if (Filename.find(Word) != std::wstring::npos)
+		{
+			return false;
+		}
+	}
+
+	return true;
+}
+
+bool
+FSyncFilter::ShouldCleanup(const FPath& Filename) const
+{
+#if UNSYNC_PLATFORM_WINDOWS
+	return ShouldCleanup(Filename.native());
+#else
+	return ShouldCleanup(Filename.wstring());
+#endif
+}
+
+bool
+FSyncFilter::ShouldCleanup(const std::wstring& Filename) const
+{
+	for (const std::wstring& Word : CleanupExcludedWords)
+	{
+		if (Filename.find(Word) != std::wstring::npos)
+		{
+			return false;
+		}
+	}
+
+	return true;
+}
+
+FPath
+FSyncFilter::Resolve(const FPath& Filename) const
+{
+	std::wstring FilenameLower = StringToLower(Filename.wstring());
+
+	FPath Result;
+	for (const auto& Alias : DfsAliases)
+	{
+		// TODO: add a case-insensitive find() helper
+		size_t Pos = FilenameLower.find(StringToLower(Alias.Source.wstring()));
+		if (Pos == 0)
+		{
+			auto Tail = (Filename.wstring().substr(Alias.Source.wstring().length() + 1));
+			Result	  = Alias.Target / Tail;
+			break;
+		}
+	}
+
+	if (Result.empty())
+	{
+		Result = Filename;
+	}
+
+	return Result;
+}
+
+}
diff --git a/Engine/Source/Programs/Unsync/Private/UnsyncFilter.h b/Engine/Source/Programs/Unsync/Private/UnsyncFilter.h
new file mode 100644
--- /dev/null
+++ b/Engine/Source/Programs/Unsync/Private/UnsyncFilter.h
@@ -0,0 +1,34 @@
+// Copyright Epic Games, Inc. All Rights Reserved.
+
+#pragma once
+
+#include "UnsyncCommon.h"
+#include "UnsyncSource.h"
+
+namespace unsync {
+
+struct FSyncFilter
+{
+	FSyncFilter() = default;
+
+	// By default all files will be included, calling this will include only files containing these substrings
+	void IncludeInSync(const std::wstring& CommaSeparatedWords);
+	void ExcludeFromSync(const std::wstring& CommaSeparatedWords);
+	void ExcludeFromCleanup(const std::wstring& CommaSeparatedWords);
+
+	bool ShouldSync(const FPath& Filename) const;
+	bool ShouldSync(const std::wstring& Filename) const;
+
+	bool ShouldCleanup(const FPath& Filename) const;
+	bool ShouldCleanup(const std::wstring& Filename) const;
+
+	FPath Resolve(const FPath& Filename) const;
+
+	std::vector<std::wstring> SyncIncludedWords;
+	std::vector<std::wstring> SyncExcludedWords;	 // any paths that contain these words will not be synced
+	std::vector<std::wstring> CleanupExcludedWords;	 // any paths that contain these words will not be deleted after sync
+
+	std::vector<FDfsAlias> DfsAliases;
+};
+
+}
diff --git a/Engine/Source/Programs/Unsync/Private/UnsyncMain.cpp b/Engine/Source/Programs/Unsync/Private/UnsyncMain.cpp
--- a/Engine/Source/Programs/Unsync/Private/UnsyncMain.cpp
+++ b/Engine/Source/Programs/Unsync/Private/UnsyncMain.cpp
@@ -19,6 +19,8 @@
 #include "UnsyncUtil.h"
 #include "UnsyncScheduler.h"
 #include "UnsyncVersion.h"
+#include "UnsyncSource.h"
+#include "UnsyncFilter.h"
 
 UNSYNC_THIRD_PARTY_INCLUDES_START
 #if UNSYNC_PLATFORM_WINDOWS
diff --git a/Engine/Source/Programs/Unsync/Private/UnsyncManifest.cpp b/Engine/Source/Programs/Unsync/Private/UnsyncManifest.cpp
--- a/Engine/Source/Programs/Unsync/Private/UnsyncManifest.cpp
+++ b/Engine/Source/Programs/Unsync/Private/UnsyncManifest.cpp
@@ -1,14 +1,15 @@
 // Copyright Epic Games, Inc. All Rights Reserved.
 
 #include "UnsyncManifest.h"
+#include "UnsyncChunking.h"
 #include "UnsyncCore.h"
 #include "UnsyncFile.h"
+#include "UnsyncFilter.h"
 #include "UnsyncHash.h"
 #include "UnsyncHashTable.h"
+#include "UnsyncScheduler.h"
 #include "UnsyncSerialization.h"
 #include "UnsyncThread.h"
-#include "UnsyncScheduler.h"
-#include "UnsyncChunking.h"
 
 UNSYNC_THIRD_PARTY_INCLUDES_START
 #include <blake3.h>
@@ -556,4 +557,255 @@ LoadOrCreateDirectoryManifest(FDirectoryManifest& Result, const FPath& Root, con
 	return true;
 }
 
+THashMap<FGenericHash, FGenericBlock>
+BuildBlockMap(const FDirectoryManifest& Manifest, bool bNeedMacroBlocks)
+{
+	THashMap<FGenericHash, FGenericBlock> Result;
+	for (const auto& It : Manifest.Files)
+	{
+		const FFileManifest& File = It.second;
+		if (bNeedMacroBlocks)
+		{
+			for (const FGenericBlock& Block : File.MacroBlocks)
+			{
+				Result[Block.HashStrong] = Block;
+			}
+		}
+		else
+		{
+			for (const FGenericBlock& Block : File.Blocks)
+			{
+				Result[Block.HashStrong] = Block;
+			}
+		}
+	}
+	return Result;
+}
+
+void
+LogManifestDiff(ELogLevel LogLevel, const FDirectoryManifest& ManifestA, const FDirectoryManifest& ManifestB)
+{
+	THashMap<FGenericHash, FGenericBlock> BlocksA = BuildBlockMap(ManifestA, false);
+	THashMap<FGenericHash, FGenericBlock> BlocksB = BuildBlockMap(ManifestB, false);
+
+	THashMap<FGenericHash, FGenericBlock> MacroBlocksA = BuildBlockMap(ManifestA, true);
+	THashMap<FGenericHash, FGenericBlock> MacroBlocksB = BuildBlockMap(ManifestB, true);
+
+	uint32 NumCommonBlocks		= 0;
+	uint64 TotalCommonBlockSize = 0;
+	uint64 TotalSizeA			= 0;
+	uint64 TotalSizeB			= 0;
+
+	uint64 PatchSizeFromAtoB = 0;
+
+	for (const auto& ItA : BlocksA)
+	{
+		TotalSizeA += ItA.second.Size;
+		auto ItB = BlocksB.find(ItA.first);
+		if (ItB != BlocksB.end())
+		{
+			NumCommonBlocks++;
+			TotalCommonBlockSize += ItA.second.Size;
+		}
+	}
+
+	for (const auto& ItB : BlocksB)
+	{
+		TotalSizeB += ItB.second.Size;
+		if (BlocksA.find(ItB.first) == BlocksA.end())
+		{
+			PatchSizeFromAtoB += ItB.second.Size;
+		}
+	}
+
+	uint32 NumCommonMacroBlocks		 = 0;
+	uint64 TotalCommonMacroBlockSize = 0;
+	for (const auto& ItA : MacroBlocksA)
+	{
+		auto ItB = MacroBlocksB.find(ItA.first);
+		if (ItB != MacroBlocksB.end())
+		{
+			NumCommonMacroBlocks++;
+			TotalCommonMacroBlockSize += ItA.second.Size;
+		}
+	}
+
+	LogPrintf(LogLevel, L"Common macro blocks: %d, %.3f MB\n", NumCommonMacroBlocks, SizeMb(TotalCommonMacroBlockSize));
+
+	LogPrintf(LogLevel,
+			  L"Common blocks: %d, %.3f MB (%.2f%% of A, %.2f%% of B)\n",
+			  NumCommonBlocks,
+			  SizeMb(TotalCommonBlockSize),
+			  100.0 * double(TotalCommonBlockSize) / double(TotalSizeA),
+			  100.0 * double(TotalCommonBlockSize) / double(TotalSizeB));
+
+	LogPrintf(LogLevel, L"Patch size: %.3f MB\n", SizeMb(PatchSizeFromAtoB));
+}
+
+void
+FilterManifest(const FSyncFilter& SyncFilter, FDirectoryManifest& Manifest)
+{
+	auto It = Manifest.Files.begin();
+	while (It != Manifest.Files.end())
+	{
+		if (SyncFilter.ShouldSync(It->first))
+		{
+			++It;
+		}
+		else
+		{
+			It = Manifest.Files.erase(It);
+		}
+	}
+}
+
+int32
+CmdInfo(const FCmdInfoOptions& Options)
+{
+	FPath DirectoryManifestPathA = IsDirectory(Options.InputA) ? (Options.InputA / ".unsync" / "manifest.bin") : Options.InputA;
+	FPath DirectoryManifestPathB = IsDirectory(Options.InputB) ? (Options.InputB / ".unsync" / "manifest.bin") : Options.InputB;
+
+	FDirectoryManifest ManifestA;
+
+	bool bManifestAValid = LoadDirectoryManifest(ManifestA, Options.InputA, DirectoryManifestPathA);
+
+	if (!bManifestAValid)
+	{
+		return 1;
+	}
+
+	LogPrintf(ELogLevel::Info, L"Manifest A: %ls\n", DirectoryManifestPathA.wstring().c_str());
+
+	if (Options.SyncFilter)
+	{
+		FilterManifest(*Options.SyncFilter, ManifestA);
+	}
+
+	{
+		UNSYNC_LOG_INDENT;
+		LogManifestInfo(ELogLevel::Info, ManifestA);
+	}
+
+	if (Options.bListFiles)
+	{
+		UNSYNC_LOG_INDENT;
+		LogManifestFiles(ELogLevel::Info, ManifestA);
+	}
+
+	if (Options.InputB.empty())
+	{
+		return 0;
+	}
+
+	LogPrintf(ELogLevel::Info, L"\n");
+
+	FDirectoryManifest ManifestB;
+
+	bool bManifestBValid = LoadDirectoryManifest(ManifestB, Options.InputB, DirectoryManifestPathB);
+
+	if (!bManifestBValid)
+	{
+		return 1;
+	}
+
+	LogPrintf(ELogLevel::Info, L"Manifest B: %ls\n", DirectoryManifestPathB.wstring().c_str());
+
+	if (Options.SyncFilter)
+	{
+		FilterManifest(*Options.SyncFilter, ManifestB);
+	}
+
+	{
+		UNSYNC_LOG_INDENT;
+		LogManifestInfo(ELogLevel::Info, ManifestB);
+	}
+	if (Options.bListFiles)
+	{
+		UNSYNC_LOG_INDENT;
+		LogManifestFiles(ELogLevel::Info, ManifestB);
+	}
+
+	LogPrintf(ELogLevel::Info, L"\n");
+	LogPrintf(ELogLevel::Info, L"Difference:\n");
+
+	{
+		UNSYNC_LOG_INDENT;
+		LogManifestDiff(ELogLevel::Info, ManifestA, ManifestB);
+	}
+
+	return 0;
+}
+
+FHash256
+ComputeSerializedManifestHash(const FDirectoryManifest& Manifest)
+{
+	FBuffer			 ManifestBuffer;
+	FVectorStreamOut ManifestStream(ManifestBuffer);
+	bool			 bSerializedOk = SaveDirectoryManifest(Manifest, ManifestStream);
+	UNSYNC_ASSERT(bSerializedOk);
+	return HashBlake3Bytes<FHash256>(ManifestBuffer.Data(), ManifestBuffer.Size());
+}
+
+FHash160
+ComputeSerializedManifestHash160(const FDirectoryManifest& Manifest)
+{
+	return ToHash160(ComputeSerializedManifestHash(Manifest));
+}
+
+bool
+MergeManifests(FDirectoryManifest& Existing, const FDirectoryManifest& Other, bool bCaseSensitive)
+{
+	if (!Existing.IsValid())
+	{
+		Existing = Other;
+		return true;
+	}
+
+	if (!AlgorithmOptionsCompatible(Existing.Algorithm, Other.Algorithm))
+	{
+		UNSYNC_ERROR("Trying to merge incompatible manifests (diff algorithm options do not match)");
+		return false;
+	}
+
+	if (bCaseSensitive)
+	{
+		// Trivial case: just replace existing entries
+		for (const auto& OtherFile : Other.Files)
+		{
+			Existing.Files[OtherFile.first] = OtherFile.second;
+		}
+	}
+	else
+	{
+		// Lookup table of lowercase -> original file name used to replace conflicting entries on non-case-sensitive filesystems
+		// TODO: Could potentially add case-sensitive/insensitive entry lookup helper functions to FDirectoryManifest itself in the future
+		std::unordered_map<std::wstring, std::wstring> ExistingFileNamesLowerCase;
+
+		for (auto& ExistingEntry : Existing.Files)
+		{
+			std::wstring FileNameLowerCase = StringToLower(ExistingEntry.first);
+			ExistingFileNamesLowerCase.insert(std::pair<std::wstring, std::wstring>(FileNameLowerCase, ExistingEntry.first));
+		}
+
+		for (const auto& OtherFile : Other.Files)
+		{
+			std::wstring OtherNameLowerCase = StringToLower(OtherFile.first);
+			auto		 LowerCaseEntry		= ExistingFileNamesLowerCase.find(OtherNameLowerCase);
+			if (LowerCaseEntry != ExistingFileNamesLowerCase.end())
+			{
+				// Remove file with conflicting case and add entry from the other manifest instead
+				const std::wstring& ExistingNameOriginalCase = LowerCaseEntry->second;
+				Existing.Files.erase(ExistingNameOriginalCase);
+
+				// Update the lookup table entry to refer to the name we're about to insert
+				ExistingFileNamesLowerCase[LowerCaseEntry->first] = OtherFile.first;
+			}
+
+			Existing.Files[OtherFile.first] = OtherFile.second;
+		}
+	}
+
+	return true;
+}
+
 }  // namespace unsync
diff --git a/Engine/Source/Programs/Unsync/Private/UnsyncManifest.h b/Engine/Source/Programs/Unsync/Private/UnsyncManifest.h
--- a/Engine/Source/Programs/Unsync/Private/UnsyncManifest.h
+++ b/Engine/Source/Programs/Unsync/Private/UnsyncManifest.h
@@ -11,6 +11,7 @@
 namespace unsync {
 
 struct FComputeBlocksParams;
+struct FSyncFilter;
 
 struct FFileManifest
 {
@@ -106,9 +107,17 @@ void			   UpdateDirectoryManifestBlocks(FDirectoryManifest& Result, const FPath&
 FDirectoryManifest CreateDirectoryManifest(const FPath& Root, const FComputeBlocksParams& Params);
 FDirectoryManifest CreateDirectoryManifestIncremental(const FPath& Root, const FComputeBlocksParams& Params);
 bool			   LoadOrCreateDirectoryManifest(FDirectoryManifest& Result, const FPath& Root, const FComputeBlocksParams& Params);
+bool			   MergeManifests(FDirectoryManifest& Existing, const FDirectoryManifest& Other, bool bCaseSensitive);
+void			   MoveCompatibleManifestBlocks(FDirectoryManifest& Manifest, FDirectoryManifest&& DonorManifest);
+bool			   AlgorithmOptionsCompatible(const FAlgorithmOptions& A, const FAlgorithmOptions& B);
 
-void MoveCompatibleManifestBlocks(FDirectoryManifest& Manifest, FDirectoryManifest&& DonorManifest);
-
-bool AlgorithmOptionsCompatible(const FAlgorithmOptions& A, const FAlgorithmOptions& B);
+struct FCmdInfoOptions
+{
+	FPath			   InputA;
+	FPath			   InputB;
+	bool			   bListFiles = false;
+	const FSyncFilter* SyncFilter = nullptr;
+};
+int32 CmdInfo(const FCmdInfoOptions& Options);
 
 }  // namespace unsync
diff --git a/Engine/Source/Programs/Unsync/Private/UnsyncSource.cpp b/Engine/Source/Programs/Unsync/Private/UnsyncSource.cpp
new file mode 100644
--- /dev/null
+++ b/Engine/Source/Programs/Unsync/Private/UnsyncSource.cpp
@@ -0,0 +1,113 @@
+// Copyright Epic Games, Inc. All Rights Reserved.
+
+#include "UnsyncSource.h"
+#include "UnsyncUtil.h"
+
+#if UNSYNC_PLATFORM_WINDOWS
+#	include <Windows.h>
+#	include <shellapi.h>
+#	include <lm.h>
+#	include <lmdfs.h>
+#	include <wincrypt.h>
+#	include <winnetwk.h>  // for WNetGetUniversalName
+#endif					   // UNSYNC_PLATFORM_WINDOWS
+
+namespace unsync {
+
+FDfsMirrorInfo
+DfsEnumerate(const FPath& Root)
+{
+	FDfsMirrorInfo Result;
+
+#if UNSYNC_PLATFORM_WINDOWS
+
+	std::wstring RootPathLower = StringToLower(Root.native());
+
+	LPWSTR RootPathCstr = (LPWSTR)Root.c_str();
+
+	DWORD ResumeHandle = 0;
+
+	std::vector<PDFS_INFO_3> InfosToFree;
+
+	PDFS_INFO_3	 BestMatchEntry = nullptr;
+	std::wstring BestMatchPath;
+
+	for (;;)
+	{
+		DWORD		EntriesRead = 0;
+		PDFS_INFO_3 DfsInfoRoot = nullptr;
+		DWORD		Res			= NetDfsEnum(RootPathCstr, 3, MAX_PREFERRED_LENGTH, (LPBYTE*)&DfsInfoRoot, &EntriesRead, &ResumeHandle);
+		if (Res == ERROR_NO_MORE_ITEMS)
+		{
+			break;
+		}
+		else if (Res == RPC_S_INVALID_NET_ADDR)
+		{
+			// Not a network share root, so nothing to do
+			break;
+		}
+		else if (Res != ERROR_SUCCESS)
+		{
+			UNSYNC_LOG(L"DFS enumeration failed with error: %d", Res);
+			break;
+		}
+
+		PDFS_INFO_3 DfsInfo = DfsInfoRoot;
+
+		for (DWORD I = 0; I < EntriesRead; I++)
+		{
+			std::wstring EntryPathLower = StringToLower(DfsInfo->EntryPath);
+
+			// entry prefix must match requested root path
+			if (RootPathLower.find(EntryPathLower) == 0 && (RootPathLower.length() > BestMatchPath.length()))
+			{
+				DWORD			  NumOnlineStorages = 0;
+				PDFS_STORAGE_INFO StorageInfo		= DfsInfo->Storage;
+				for (DWORD J = 0; J < DfsInfo->NumberOfStorages; J++)
+				{
+					if (StorageInfo->State != DFS_STORAGE_STATE_OFFLINE)
+					{
+						NumOnlineStorages++;
+					}
+					++StorageInfo;
+				}
+
+				BestMatchPath  = DfsInfo->EntryPath;
+				BestMatchEntry = DfsInfo;
+			}
+
+			++DfsInfo;
+		}
+
+		InfosToFree.push_back(DfsInfoRoot);
+	}
+
+	if (BestMatchEntry)
+	{
+		Result.Root					  = BestMatchPath;
+		PDFS_STORAGE_INFO StorageInfo = BestMatchEntry->Storage;
+		Result.Storages.reserve(BestMatchEntry->NumberOfStorages);
+		for (DWORD J = 0; J < BestMatchEntry->NumberOfStorages; J++)
+		{
+			if (StorageInfo->State != DFS_STORAGE_STATE_OFFLINE)
+			{
+				FDfsStorageInfo ResultEntry;
+				ResultEntry.Server = StorageInfo->ServerName;
+				ResultEntry.Share  = StorageInfo->ShareName;
+				Result.Storages.push_back(ResultEntry);
+			}
+			++StorageInfo;
+		}
+	}
+
+	for (auto It : InfosToFree)
+	{
+		NetApiBufferFree(It);
+	}
+
+#endif	// UNSYNC_PLATFORM_WINDOWS
+
+	return Result;
+}
+
+}  // namespace unsync
diff --git a/Engine/Source/Programs/Unsync/Private/UnsyncSource.h b/Engine/Source/Programs/Unsync/Private/UnsyncSource.h
new file mode 100644
--- /dev/null
+++ b/Engine/Source/Programs/Unsync/Private/UnsyncSource.h
@@ -0,0 +1,45 @@
+// Copyright Epic Games, Inc. All Rights Reserved.
+
+#pragma once
+
+#include "UnsyncCommon.h"
+
+namespace unsync {
+
+enum class ESourceType : uint8 {
+	Unknown,
+	FileSystem,
+	Server,
+	ServerWithManifestHash,
+};
+
+struct FSourcePath
+{
+	ESourceType Type = ESourceType::Unknown;
+	std::string Location;  // utf-8
+
+	bool operator==(const FSourcePath& Other) const { return Type == Other.Type && Location == Other.Location; }
+};
+
+// Returns a list of alternative DFS paths for a given root
+struct FDfsStorageInfo
+{
+	std::wstring Server;
+	std::wstring Share;
+
+	bool IsValid() const { return !Server.empty() && !Share.empty(); }
+};
+struct FDfsMirrorInfo
+{
+	std::wstring				 Root;
+	std::vector<FDfsStorageInfo> Storages;
+};
+FDfsMirrorInfo DfsEnumerate(const FPath& Root);
+
+struct FDfsAlias
+{
+	FPath Source;
+	FPath Target;
+};
+
+}  // namespace unsync
diff --git a/Engine/Source/Programs/Unsync/Private/UnsyncUtil.cpp b/Engine/Source/Programs/Unsync/Private/UnsyncUtil.cpp
--- a/Engine/Source/Programs/Unsync/Private/UnsyncUtil.cpp
+++ b/Engine/Source/Programs/Unsync/Private/UnsyncUtil.cpp
@@ -276,102 +276,6 @@ StringEscape(const std::string_view Input)
 	return Result;
 }
 
-FDfsMirrorInfo
-DfsEnumerate(const FPath& Root)
-{
-	FDfsMirrorInfo Result;
-
-#if UNSYNC_PLATFORM_WINDOWS
-
-	std::wstring RootPathLower = StringToLower(Root.native());
-
-	LPWSTR RootPathCstr = (LPWSTR)Root.c_str();
-
-	DWORD ResumeHandle = 0;
-
-	std::vector<PDFS_INFO_3> InfosToFree;
-
-	PDFS_INFO_3	 BestMatchEntry = nullptr;
-	std::wstring BestMatchPath;
-
-	for (;;)
-	{
-		DWORD		EntriesRead = 0;
-		PDFS_INFO_3 DfsInfoRoot = nullptr;
-		DWORD		Res			= NetDfsEnum(RootPathCstr, 3, MAX_PREFERRED_LENGTH, (LPBYTE*)&DfsInfoRoot, &EntriesRead, &ResumeHandle);
-		if (Res == ERROR_NO_MORE_ITEMS)
-		{
-			break;
-		}
-		else if (Res == RPC_S_INVALID_NET_ADDR)
-		{
-			// Not a network share root, so nothing to do
-			break;
-		}
-		else if (Res != ERROR_SUCCESS)
-		{
-			UNSYNC_LOG(L"DFS enumeration failed with error: %d", Res);
-			break;
-		}
-
-		PDFS_INFO_3 DfsInfo = DfsInfoRoot;
-
-		for (DWORD I = 0; I < EntriesRead; I++)
-		{
-			std::wstring EntryPathLower = StringToLower(DfsInfo->EntryPath);
-
-			// entry prefix must match requested root path
-			if (RootPathLower.find(EntryPathLower) == 0 && (RootPathLower.length() > BestMatchPath.length()))
-			{
-				DWORD			  NumOnlineStorages = 0;
-				PDFS_STORAGE_INFO StorageInfo		= DfsInfo->Storage;
-				for (DWORD J = 0; J < DfsInfo->NumberOfStorages; J++)
-				{
-					if (StorageInfo->State != DFS_STORAGE_STATE_OFFLINE)
-					{
-						NumOnlineStorages++;
-					}
-					++StorageInfo;
-				}
-
-				BestMatchPath  = DfsInfo->EntryPath;
-				BestMatchEntry = DfsInfo;
-			}
-
-			++DfsInfo;
-		}
-
-		InfosToFree.push_back(DfsInfoRoot);
-	}
-
-	if (BestMatchEntry)
-	{
-		Result.Root					  = BestMatchPath;
-		PDFS_STORAGE_INFO StorageInfo = BestMatchEntry->Storage;
-		Result.Storages.reserve(BestMatchEntry->NumberOfStorages);
-		for (DWORD J = 0; J < BestMatchEntry->NumberOfStorages; J++)
-		{
-			if (StorageInfo->State != DFS_STORAGE_STATE_OFFLINE)
-			{
-				FDfsStorageInfo ResultEntry;
-				ResultEntry.Server = StorageInfo->ServerName;
-				ResultEntry.Share  = StorageInfo->ShareName;
-				Result.Storages.push_back(ResultEntry);
-			}
-			++StorageInfo;
-		}
-	}
-
-	for (auto It : InfosToFree)
-	{
-		NetApiBufferFree(It);
-	}
-
-#endif	// UNSYNC_PLATFORM_WINDOWS
-
-	return Result;
-}
-
 FPath
 GetUniversalPath(const FPath& Path)
 {
diff --git a/Engine/Source/Programs/Unsync/Private/UnsyncUtil.h b/Engine/Source/Programs/Unsync/Private/UnsyncUtil.h
--- a/Engine/Source/Programs/Unsync/Private/UnsyncUtil.h
+++ b/Engine/Source/Programs/Unsync/Private/UnsyncUtil.h
@@ -144,27 +144,6 @@ std::wstring StringToUpper(const std::wstring& Input);
 
 std::string StringEscape(const std::string_view Input);
 
-// Returns a list of alternative DFS paths for a given root
-struct FDfsStorageInfo
-{
-	std::wstring Server;
-	std::wstring Share;
-
-	bool IsValid() const { return !Server.empty() && !Share.empty(); }
-};
-struct FDfsMirrorInfo
-{
-	std::wstring				 Root;
-	std::vector<FDfsStorageInfo> Storages;
-};
-FDfsMirrorInfo DfsEnumerate(const FPath& Root);
-
-struct FDfsAlias
-{
-	FPath Source;
-	FPath Target;
-};
-
 // Takes a drive-based path (e.g. P:/Foo/Bar) and converts it to universal form (e.g. //server/Foo/Bar), if possible.
 // Otherwise, returns original path.
 FPath GetUniversalPath(const FPath& Path);
diff --git a/Engine/Source/Programs/Unsync/Private/UnsyncVersion.h b/Engine/Source/Programs/Unsync/Private/UnsyncVersion.h
--- a/Engine/Source/Programs/Unsync/Private/UnsyncVersion.h
+++ b/Engine/Source/Programs/Unsync/Private/UnsyncVersion.h
@@ -1,5 +1,7 @@
 // Copyright Epic Games, Inc. All Rights Reserved.
 
+#pragma once
+
 #include "UnsyncCommon.h"
 
 #define UNSYNC_VERSION_STR "1.0.78"
