unsync - Add an abstraction for remote file system listing and manifest file downloads

diff --git a/Engine/Source/Programs/Unsync/Private/UnsyncCore.cpp b/Engine/Source/Programs/Unsync/Private/UnsyncCore.cpp
--- a/Engine/Source/Programs/Unsync/Private/UnsyncCore.cpp
+++ b/Engine/Source/Programs/Unsync/Private/UnsyncCore.cpp
@@ -797,20 +797,27 @@ struct FPooledProxy
 	std::unique_ptr<FProxy> Proxy;
 };
 
-struct FRemoteManifestInfo
+struct FRemoteFileInfo
 {
 	FPath							   Path;
 	ProxyQuery::FDirectoryListingEntry Entry;
 };
 
-static TResult<FRemoteManifestInfo>
-FindRemoteUnsyncManifest(FHttpConnection& Connection, const FAuthDesc* AuthDesc, const FPath& RootDirectory)
+struct FFoundManifest
+{
+	FRemoteFileInfo				Manifest;
+	std::vector<FRemoteFileInfo> PackDataFiles;
+	std::vector<FRemoteFileInfo> PackIndexFiles;
+};
+
+static TResult<FFoundManifest>
+FindUnsyncManifest(FProxyFileSystem& FileSystem)
 {
 	using FDirectoryListing		 = ProxyQuery::FDirectoryListing;
 	using FDirectoryListingEntry = ProxyQuery::FDirectoryListingEntry;
 
-	std::string RootDirectoryUtf8 = ConvertWideToUtf8(RootDirectory.wstring());
-	TResult<FDirectoryListing> RootDirectoryListingResult = ProxyQuery::ListDirectory(Connection, AuthDesc, RootDirectoryUtf8);
+	TResult<FDirectoryListing> RootDirectoryListingResult = FileSystem.ListDirectory("");
+
 	UNSYNC_RETURN_ON_ERROR(RootDirectoryListingResult);
 
 	std::optional<FDirectoryListingEntry> FoundUnsyncSubdirectory;
@@ -831,70 +838,102 @@ FindRemoteUnsyncManifest(FHttpConnection& Connection, const FAuthDesc* AuthDesc,
 
 	if (FoundUnsyncManifestDotfile && !FoundUnsyncSubdirectory)
 	{
-		FRemoteManifestInfo Result;
-		Result.Path	 = RootDirectory / FoundUnsyncManifestDotfile->Name;
-		Result.Entry = FoundUnsyncManifestDotfile.value();
+		FFoundManifest Result;
+		Result.Manifest.Path = FoundUnsyncManifestDotfile->Name;
+		Result.Manifest.Entry = *FoundUnsyncManifestDotfile;
 		return ResultOk(Result);
 	}
 
 	if (FoundUnsyncSubdirectory)
 	{
-		FPath UnsyncDirectory = RootDirectory / ".unsync";
-		std::string UnsyncDirectoryUtf8 = ConvertWideToUtf8(UnsyncDirectory.wstring());
+		std::string				   UnsyncDirectory				= ".unsync";
+		TResult<FDirectoryListing> UnsyncDirectoryListingResult = FileSystem.ListDirectory(UnsyncDirectory);
 
-		TResult<FDirectoryListing> UnsyncDirectoryListingResult = ProxyQuery::ListDirectory(Connection, AuthDesc, UnsyncDirectoryUtf8);
 		UNSYNC_RETURN_ON_ERROR(UnsyncDirectoryListingResult);
 
+		FFoundManifest Result;
+
 		for (const FDirectoryListingEntry& Entry : UnsyncDirectoryListingResult.GetData().Entries)
 		{
 			if (Entry.Name == "manifest.bin" && !Entry.bDirectory)
 			{
-				FRemoteManifestInfo Result;
-				Result.Path	 = UnsyncDirectory / Entry.Name;
-				Result.Entry = Entry;
-				return ResultOk(Result);
+				Result.Manifest.Path = FPath(UnsyncDirectory) / Entry.Name;
+				Result.Manifest.Entry = Entry;
+			}
+
+			if (Entry.Name == "pack" && Entry.bDirectory)
+			{
+				FPath					   PackDirectory		 = FPath(UnsyncDirectory) / "pack";
+				TResult<FDirectoryListing> PackDirectoryListring = FileSystem.ListDirectory(ToString(PackDirectory));
+				if (const FDirectoryListing* Listing = PackDirectoryListring.TryData())
+				{
+					for (const FDirectoryListingEntry& PackEntry : Listing->Entries)
+					{
+						if (PackEntry.bDirectory)
+						{
+							continue;
+						}
+
+						FRemoteFileInfo Info;
+						Info.Path = PackDirectory / PackEntry.Name;
+						Info.Entry = PackEntry;
+
+						if (PackEntry.Name.ends_with(".unsync_pack"))
+						{
+							Result.PackDataFiles.push_back(Info);
+						}
+
+						if (PackEntry.Name.ends_with(".unsync_index"))
+						{
+							Result.PackIndexFiles.push_back(Info);
+						}
+					}
+				}
 			}
 		}
+
+		if (!Result.Manifest.Path.empty())
+		{
+			return ResultOk(std::move(Result));
+		}
 	}
 
-	return AppError("Could not find remote unsync manifest file");
+	return AppError("Could not find unsync manifest file");
 }
 
 static bool
-DownloadFileIfPossiblyDifferent(FHttpConnection&		   Connection,
-								const FAuthDesc*		   AuthDesc,
-								const FRemoteManifestInfo& Source,
-								const FPath&			   Target,
-								EFileMode				   TargetFileMode)
+CopyFileIfPossiblyDifferent(FProxyFileSystem&	   FileSystem,
+							const FRemoteFileInfo& Source,
+							const FPath&		   Target,
+							EFileMode			   TargetFileMode = EFileMode::CreateWriteOnly)
 {
 	FFileAttributes TargetAttr = GetFileAttrib(Target);
 	if (Source.Entry.Size != TargetAttr.Size || Source.Entry.Mtime != TargetAttr.Mtime)
 	{
-		UNSYNC_VERBOSE(L"Downloading '%ls'", Source.Path.wstring().c_str());
-		std::string SourceUtf8 = ConvertWideToUtf8(Source.Path.wstring());
+		UNSYNC_VERBOSE(L"Reading '%ls'", Source.Path.wstring().c_str());
 
-		TResult<FBuffer> DownloadResult = ProxyQuery::DownloadFile(Connection, AuthDesc, SourceUtf8);
-		if (DownloadResult.IsError())
+		TResult<FBuffer> FileBuffer = FileSystem.ReadFile(ToString(Source.Path));
+		if (FileBuffer.IsError())
 		{
-			LogError(DownloadResult.GetError());
+			LogError(FileBuffer.GetError());
 			return false;
 		}
 
-		const FBuffer& FileBuffer = DownloadResult.GetData();
-		if (FileBuffer.Size() != Source.Entry.Size)
+		if (FileBuffer->Size() != Source.Entry.Size)
 		{
-			UNSYNC_ERROR(L"Downloaded file size mismatch. Expected %llu, actual %llu.", llu(Source.Entry.Size), llu(FileBuffer.Size()));
+			UNSYNC_ERROR(L"Read file size mismatch. Expected %llu, actual %llu.", llu(Source.Entry.Size), llu(FileBuffer->Size()));
 			return false;
 		}
 
-		bool bFileWritten = WriteBufferToFile(Target, FileBuffer, TargetFileMode);
+		const bool bFileWritten = WriteBufferToFile(Target, *FileBuffer, TargetFileMode);
 		if (!bFileWritten)
 		{
-			UNSYNC_ERROR(L"Failed to write downloaded file '%ls'", Target.wstring().c_str());
+			UNSYNC_ERROR(L"Failed to write file '%ls'", Target.wstring().c_str());
 			return false;
 		}
 
-		SetFileMtime(Target, Source.Entry.Mtime, /*allow in dry run*/ true);
+		const bool bAllowInDryRun = true;
+		SetFileMtime(Target, Source.Entry.Mtime, bAllowInDryRun);
 	}
 
 	return true;
@@ -902,104 +941,66 @@ DownloadFileIfPossiblyDifferent(FHttpConnection&		   Connection,
 
 static bool
 LoadAndMergeSourceManifest(FDirectoryManifest& Output,
-						   FProxyPool&		   ProxyPool,
+						   FProxyFileSystem&   ProxyFileSystem,
 						   const FPath&		   SourcePath,
 						   const FPath&		   TempPath,
 						   FSyncFilter*		   SyncFilter,
-						   const FPath&		   SourceManifestOverride,
 						   bool				   bCaseSensitiveTargetFileSystem)
 {
-	const FRemoteProtocolFeatures& ProxyFeatures = ProxyPool.GetFeatures();
-
-	const bool bDownloadManifestFromProxy =
-		ProxyPool.RemoteDesc.Protocol == EProtocolFlavor::Unsync && ProxyFeatures.bDirectoryListing && ProxyFeatures.bFileDownload;
-
-	UNSYNC_VERBOSE2(L"LoadAndMergeSourceManifest: '%ls' (%hs)",
-					SourcePath.wstring().c_str(),
-					bDownloadManifestFromProxy ? "download" : "filesystem");
-
-	auto ResolvePath = [SyncFilter, bDownloadManifestFromProxy](const FPath& Filename) -> FPath
-	{ return (SyncFilter && !bDownloadManifestFromProxy) ? SyncFilter->Resolve(Filename) : Filename; };
+	UNSYNC_VERBOSE2(L"LoadAndMergeSourceManifest: '%ls'", SourcePath.wstring().c_str());
 
 	FDirectoryManifest LoadedManifest;
 
-	FPath SourceManifestPath = SourcePath / ".unsync" / "manifest.bin";
-
-	SourceManifestPath = ResolvePath(SourceManifestPath);
-
-	if (!SourceManifestOverride.empty())
-	{
-		SourceManifestPath = SourceManifestOverride;
-	}
+	FPath SourceManifestRelativePath = FPath(".unsync") / "manifest.bin";
+	FPath SourceManifestPath		 = SourcePath / SourceManifestRelativePath;
 
 	FHash128 SourcePathHash =
 		HashBlake3Bytes<FHash128>((const uint8*)SourcePath.native().c_str(), SourcePath.native().length() * sizeof(SourcePath.native()[0]));
 
-	std::string SourcePathHashStr	   = BytesToHexString(SourcePathHash.Data, sizeof(SourcePathHash.Data));
-	FPath		SourceManifestTempPath = TempPath / SourcePathHashStr;
+	std::string SourcePathHashStr  = BytesToHexString(SourcePathHash.Data, sizeof(SourcePathHash.Data));
+	FPath		CachedManifestPath = TempPath / SourcePathHashStr;
 
 	LogGlobalStatus(L"Caching source manifest");
 	UNSYNC_VERBOSE(L"Caching source manifest");
 
-	UNSYNC_VERBOSE(L" Source '%ls'", SourceManifestPath.wstring().c_str());
-	UNSYNC_VERBOSE(L" Target '%ls'", SourceManifestTempPath.wstring().c_str());
+	UNSYNC_LOG_INDENT;
+	UNSYNC_VERBOSE(L"Source '%ls'", SourceManifestPath.wstring().c_str());
+	UNSYNC_VERBOSE(L"Target '%ls'", CachedManifestPath.wstring().c_str());
 
-	if (bDownloadManifestFromProxy)
+	TResult<FFoundManifest> FindManifestResult = FindUnsyncManifest(ProxyFileSystem);
+	if (FindManifestResult.IsError())
 	{
-		UNSYNC_LOG_INDENT;
-
-		FHttpConnection Connection = FHttpConnection::CreateDefaultHttps(ProxyPool.RemoteDesc);
-
-		TResult<FRemoteManifestInfo> FoundManifestResult = FindRemoteUnsyncManifest(Connection, ProxyPool.AuthDesc, SourcePath);
-		UNSYNC_UNUSED(FoundManifestResult);
-
-		if (FoundManifestResult.IsError())
-		{
-			LogError(FoundManifestResult.GetError());
-			return false;
-		}
-
-		const FRemoteManifestInfo& FoundManifest = FoundManifestResult.GetData();
-
-		bool bDownloadedOk = DownloadFileIfPossiblyDifferent(Connection,
-															 ProxyPool.AuthDesc,
-															 FoundManifest,
-															 SourceManifestTempPath,
-															 EFileMode::CreateWriteOnly | EFileMode::IgnoreDryRun);
-
-		if (!bDownloadedOk)
-		{
-			UNSYNC_ERROR(L"Failed to download manifest file '%ls'", SourceManifestPath.wstring().c_str());
-			return false;
-		}
+		LogError(FindManifestResult.GetError());
+		return false;
 	}
-	else
-	{
-		UNSYNC_LOG_INDENT;
-		if (!PathExists(SourceManifestPath))
-		{
-			UNSYNC_ERROR(L"Source manifest '%ls' does not exist", SourceManifestPath.wstring().c_str());
-			return false;
-		}
 
-		std::error_code CopyErrorCode = CopyFileIfNewer(SourceManifestPath, SourceManifestTempPath);
-		if (CopyErrorCode)
-		{
-			UNSYNC_LOG(L"Failed to copy manifest '%ls' to '%ls'",
-					   SourceManifestPath.wstring().c_str(),
-					   SourceManifestTempPath.wstring().c_str());
-			UNSYNC_ERROR(L"%hs (%d)", CopyErrorCode.message().c_str(), CopyErrorCode.value());
-			return false;
-		}
+	if (!CopyFileIfPossiblyDifferent(ProxyFileSystem,
+		FindManifestResult->Manifest,
+		CachedManifestPath,
+		EFileMode::CreateReadWrite | EFileMode::IgnoreDryRun))
+	{
+		return false;
 	}
 
-	if (!LoadDirectoryManifest(LoadedManifest, SourcePath, SourceManifestTempPath))
+	if (!LoadDirectoryManifest(LoadedManifest, SourcePath, CachedManifestPath))
 	{
 		UNSYNC_ERROR(L"Failed to load source directory manifest '%ls'", SourceManifestPath.wstring().c_str());
 
 		return false;
 	}
 
+	#if 0
+	if (!FindManifestResult->PackIndexFiles.empty())
+	{
+		UNSYNC_VERBOSE(L"Loading pack index database");
+		UNSYNC_LOG_INDENT;
+		for (const FRemoteFileInfo& FileInfo : FindManifestResult->PackIndexFiles)
+		{
+			TResult<FBuffer> FileBuffer = ProxyFileSystem.ReadFile(ToString(FileInfo.Path));
+		}
+	}
+	#endif
+
 	if (Output.IsValid() && !AlgorithmOptionsCompatible(Output.Algorithm, LoadedManifest.Algorithm))
 	{
 		UNSYNC_ERROR(L"Can't merge manifest '%ls' as it uses different algorithm options", SourcePath.wstring().c_str());
@@ -1093,7 +1094,7 @@ struct FFileSyncTaskBatch
 };
 
 static void
-DeleteOldFilesInDirectory(FPath& Path, uint32 MaxFilesToKeep)
+DeleteOldFilesInDirectory(const FPath& Path, uint32 MaxFilesToKeep)
 {
 	struct FEntry
 	{
@@ -1156,7 +1157,6 @@ SyncDirectory(const FSyncDirectoryOptions& SyncOptions)
 	const FPath	 SourcePath				= bFileSystemSource ? std::filesystem::absolute(SyncOptions.Source) : SyncOptions.Source;
 	const FPath	 BasePath				= std::filesystem::absolute(SyncOptions.Base);
 	const FPath	 TargetPath				= std::filesystem::absolute(SyncOptions.Target);
-	const FPath& SourceManifestOverride = SyncOptions.SourceManifestOverride;
 
 	FSyncFilter* SyncFilter = SyncOptions.SyncFilter;
 
@@ -1169,14 +1169,14 @@ SyncDirectory(const FSyncDirectoryOptions& SyncOptions)
 		UNSYNC_LOG(L"Unnecessary files will be deleted after sync (cleanup mode)");
 	}
 
-	FPath BaseManifestRoot = BasePath / ".unsync";
-	FPath BaseManifestPath = BaseManifestRoot / "manifest.bin";
+	const FPath BaseManifestRoot = BasePath / ".unsync";
+	const FPath BaseManifestPath = BaseManifestRoot / "manifest.bin";
 
-	FPath TargetManifestRoot = TargetPath / ".unsync";
-	FPath TargetManifestPath = TargetManifestRoot / "manifest.bin";
-	FPath TargetTempPath	 = TargetManifestRoot / "temp";
+	const FPath TargetManifestRoot = TargetPath / ".unsync";
+	const FPath TargetManifestPath = TargetManifestRoot / "manifest.bin";
+	const FPath TargetTempPath	 = TargetManifestRoot / "temp";
 
-	bool bTempDirectoryExists = (PathExists(TargetTempPath) && IsDirectory(TargetTempPath)) || CreateDirectories(TargetTempPath);
+	const bool bTempDirectoryExists = (PathExists(TargetTempPath) && IsDirectory(TargetTempPath)) || CreateDirectories(TargetTempPath);
 
 	if (!bTempDirectoryExists)
 	{
@@ -1192,8 +1192,8 @@ SyncDirectory(const FSyncDirectoryOptions& SyncOptions)
 		DeleteOldFilesInDirectory(TargetTempPath, MaxFilesToKeep);
 	}
 
-	FPath		  LogFilePath = TargetManifestRoot / L"unsync.log";
-	FLogFileScope LogFileScope(LogFilePath.wstring().c_str());
+	const FPath			LogFilePath = TargetManifestRoot / L"unsync.log";
+	const FLogFileScope LogFileScope(LogFilePath.wstring().c_str());
 	SetCrashDumpPath(TargetManifestRoot);
 
 	auto ShouldSync = [SyncFilter](const FPath& Filename) -> bool {
@@ -1244,8 +1244,41 @@ SyncDirectory(const FSyncDirectoryOptions& SyncOptions)
 			return false;
 		}
 	}
+	else if (!SyncOptions.SourceManifestOverride.empty())
+	{
+		bSourceManifestOk = LoadDirectoryManifest(SourceDirectoryManifest, SourcePath, SyncOptions.SourceManifestOverride);
+		if (!bSourceManifestOk)
+		{
+			UNSYNC_ERROR(L"Could not load explicit manifest file");
+			return false;
+		}
+	}
 	else
 	{
+		std::unique_ptr<FProxyFileSystem> ProxyFileSystem;
+		if (bServerSource)
+		{
+			const FRemoteProtocolFeatures& RemoteFeatures = ProxyPool.GetFeatures();
+
+			if (!RemoteFeatures.bDirectoryListing)
+			{
+				UNSYNC_ERROR(L"Remote server does not support directory listing");
+				return false;
+			}
+
+			if (!RemoteFeatures.bFileDownload)
+			{
+				UNSYNC_ERROR(L"Remote server does not support direct file downloads");
+				return false;
+			}
+
+			ProxyFileSystem = std::make_unique<FRemoteFileSystem>(ToString(SyncOptions.Source), ProxyPool);
+		}
+		else
+		{
+			ProxyFileSystem = std::make_unique<FPhysicalFileSystem>(SyncOptions.Source);
+		}
+
 		std::vector<FPath> AllSources;
 		AllSources.push_back(SourcePath);
 		for (const FPath& OverlayPath : SyncOptions.Overlays)
@@ -1256,11 +1289,10 @@ SyncDirectory(const FSyncDirectoryOptions& SyncOptions)
 		for (const FPath& ThisSourcePath : AllSources)
 		{
 			if (!LoadAndMergeSourceManifest(SourceDirectoryManifest,
-											ProxyPool,
+											*ProxyFileSystem,
 											ThisSourcePath,
 											TargetTempPath,
 											SyncFilter,
-											SourceManifestOverride,
 											bCaseSensitiveTargetFileSystem))
 			{
 				return false;
@@ -1683,7 +1715,7 @@ SyncDirectory(const FSyncDirectoryOptions& SyncOptions)
 		// Throttle background tasks by trying to keep them to some sensible memory budget. Best effort only, not a hard limit.
 		const uint64 BackgroundTaskMemoryBudget	 = SyncOptions.BackgroundTaskMemoryBudget;
 		const uint64 TargetTotalSizePerTaskBatch = BackgroundTaskMemoryBudget;
-		const uint64 MaxFilesPerTaskBatch		 = 1000;
+		const uint64 MaxFilesPerTaskBatch		 = SyncOptions.MaxFilesPerTask;
 
 		UNSYNC_VERBOSE2(L"Background task memory budget: %llu GB", BackgroundTaskMemoryBudget >> 30);
 
diff --git a/Engine/Source/Programs/Unsync/Private/UnsyncCore.h b/Engine/Source/Programs/Unsync/Private/UnsyncCore.h
--- a/Engine/Source/Programs/Unsync/Private/UnsyncCore.h
+++ b/Engine/Source/Programs/Unsync/Private/UnsyncCore.h
@@ -241,6 +241,7 @@ struct FSyncDirectoryOptions
 	bool			   bFullDifference = true;	// whether to run full file difference algorithm, even when there is an existing manifest
 	bool			   bCheckAvailableSpace		  = true;  // whether to abort the sync if target path does not have enough available space
 	uint64			   BackgroundTaskMemoryBudget = 2_GB;
+	uint64			   MaxFilesPerTask			  = 1000;
 };
 
 bool SyncDirectory(const FSyncDirectoryOptions& SyncOptions);
diff --git a/Engine/Source/Programs/Unsync/Private/UnsyncFile.cpp b/Engine/Source/Programs/Unsync/Private/UnsyncFile.cpp
--- a/Engine/Source/Programs/Unsync/Private/UnsyncFile.cpp
+++ b/Engine/Source/Programs/Unsync/Private/UnsyncFile.cpp
@@ -110,6 +110,13 @@ std::filesystem::file_time_type FromWindowsFileTime(uint64 Ticks)
 
 FPath
 GetRelativePath(const FPath& Path, const FPath& Base)
+{
+	FPathStringView ResultView = GetRelativePathView(Path, Base);
+	return ResultView;
+}
+
+FPathStringView
+GetRelativePathView(const FPath& Path, const FPath& Base)
 {
 	// Try a trivial case first, without touching the filesystem
 	FPathStringView PathView = RemoveExtendedPathPrefix(Path);
@@ -124,12 +131,52 @@ GetRelativePath(const FPath& Path, const FPath& Base)
 		{
 			RelativePath = RelativePath.substr(1);
 		}
-		return FPath(RelativePath);
+		return RelativePath;
 	}
 
 	return {};
 }
 
+void
+ConvertDirectorySeparatorsToNative(std::string& Path)
+{
+	std::replace_if(
+		Path.begin(),
+		Path.end(),
+		[](char C) { return C == '/' || C == '\\'; },
+		char(FPath::preferred_separator));
+}
+
+void
+ConvertDirectorySeparatorsToUnix(std::string& Path)
+{
+	std::replace_if(
+		Path.begin(),
+		Path.end(),
+		[](char C) { return C == '\\'; },
+		char('/'));
+}
+
+void
+ConvertDirectorySeparatorsToNative(std::wstring& Path)
+{
+	std::replace_if(
+		Path.begin(),
+		Path.end(),
+		[](wchar_t C) { return C == '/' || C == '\\'; },
+		wchar_t(FPath::preferred_separator));
+}
+
+void
+ConvertDirectorySeparatorsToUnix(std::wstring& Path)
+{
+	std::replace_if(
+		Path.begin(),
+		Path.end(),
+		[](wchar_t C) { return C == '\\'; },
+		wchar_t('/'));
+}
+
 std::error_code
 CopyFileIfNewer(const FPath& Source, const FPath& Target)
 {
diff --git a/Engine/Source/Programs/Unsync/Private/UnsyncFile.h b/Engine/Source/Programs/Unsync/Private/UnsyncFile.h
--- a/Engine/Source/Programs/Unsync/Private/UnsyncFile.h
+++ b/Engine/Source/Programs/Unsync/Private/UnsyncFile.h
@@ -475,10 +475,16 @@ bool			FileCopy(const FPath& From, const FPath& To, std::error_code& OutErrorCod
 bool			FileCopyOverwrite(const FPath& From, const FPath& To, std::error_code& OutErrorCode);
 bool			FileRemove(const FPath& Path, std::error_code& OutErrorCode);
 FPath			GetRelativePath(const FPath& Path, const FPath& Base);
+FPathStringView GetRelativePathView(const FPath& Path, const FPath& Base);
 std::error_code CopyFileIfNewer(const FPath& Source, const FPath& Target);
 bool			IsNonCaseSensitiveFileSystem(const FPath& ExistingPath);
 bool			IsCaseSensitiveFileSystem(const FPath& ExistingPath);
 
+void ConvertDirectorySeparatorsToNative(std::string& Path);
+void ConvertDirectorySeparatorsToNative(std::wstring& Path);
+void ConvertDirectorySeparatorsToUnix(std::string& Path);
+void ConvertDirectorySeparatorsToUnix(std::wstring& Path);
+
 // Returns number of bytes that can be written to the given path.
 // Returns ~0ull if the available space could not be determined.
 uint64 GetAvailableDiskSpace(const FPath& Path);
diff --git a/Engine/Source/Programs/Unsync/Private/UnsyncProxy.cpp b/Engine/Source/Programs/Unsync/Private/UnsyncProxy.cpp
--- a/Engine/Source/Programs/Unsync/Private/UnsyncProxy.cpp
+++ b/Engine/Source/Programs/Unsync/Private/UnsyncProxy.cpp
@@ -9,10 +9,14 @@
 #include "UnsyncProgress.h"
 #include "UnsyncScheduler.h"
 
-#include <fmt/format.h>
 #include <atomic>
 #include <json11.hpp>
 
+#include <fmt/format.h>
+#if __has_include(<fmt/xchar.h>)
+#	include <fmt/xchar.h>
+#endif
+
 namespace unsync {
 
 struct FUnsyncProtocolImpl : FRemoteProtocolBase
@@ -958,6 +962,14 @@ FProxyPool::FProxyPool(const FRemoteDesc& InRemoteDesc, const FAuthDesc* InAuthD
 		return;
 	}
 
+	auto CreateHttpConnection = [Remote = RemoteDesc]
+	{
+		FTlsClientSettings TlsSettings = Remote.GetTlsClientSettings();
+		return new FHttpConnection(Remote.Host.Address, Remote.Host.Port, Remote.TlsRequirement, TlsSettings);
+	};
+
+	HttpPool.emplace(CreateHttpConnection);
+
 	if (RemoteDesc.Protocol == EProtocolFlavor::Unsync)
 	{
 		UNSYNC_VERBOSE(L"Connecting to %hs server '%hs:%d' ...",
@@ -1028,6 +1040,24 @@ FProxyPool::Dealloc(std::unique_ptr<FProxy>&& Proxy)
 	}
 }
 
+std::unique_ptr<FHttpConnection>
+FProxyPool::AllocHttp()
+{
+	if (!bValid || !HttpPool)
+	{
+		return nullptr;
+	}
+	return HttpPool->Acquire();
+}
+void
+FProxyPool::DeallocHttp(std::unique_ptr<FHttpConnection>&& Connection)
+{
+	if (Connection && HttpPool)
+	{
+		HttpPool->Release(std::move(Connection));
+	}
+}
+
 void
 FProxyPool::Invalidate()
 {
@@ -1063,4 +1093,72 @@ FProxyPool::InitRequestMap(EStrongHashAlgorithmID InStrongHasher)
 	RequestMap.Init(InStrongHasher);
 }
 
+FPhysicalFileSystem::FPhysicalFileSystem(const FPath& InRoot) : Root(InRoot)
+{
+}
+
+TResult<FProxyDirectoryListing>
+FPhysicalFileSystem::ListDirectory(const std::string_view RelativePath)
+{
+	FProxyDirectoryListing Result;
+
+	std::wstring RelativePathWide = ConvertUtf8ToWide(RelativePath);
+	FPath		 FullPath		  = Root / FPath(RelativePathWide);
+
+	for (const std::filesystem::directory_entry& Dir : DirectoryScan(FullPath))
+	{
+		FPath FileName = Dir.path().filename();
+
+		FProxyDirectoryEntry Entry;
+		Entry.bDirectory = Dir.is_directory();
+		Entry.Size		 = Entry.bDirectory ? 0 : Dir.file_size();
+		Entry.Mtime		 = ToWindowsFileTime(Dir.last_write_time());
+		Entry.Name		 = ToString(FileName);
+
+		Result.Entries.emplace_back(std::move(Entry));
+	}
+
+	return ResultOk(std::move(Result));
+}
+
+TResult<FBuffer>
+FPhysicalFileSystem::ReadFile(const std::string_view RelativePath)
+{
+	std::wstring RelativePathWide = ConvertUtf8ToWide(RelativePath);
+	FPath		 FullPath		  = Root / FPath(RelativePathWide);
+	FBuffer		 Buffer			  = ReadFileToBuffer(FullPath);
+
+	if (Buffer.Empty())
+	{
+		// TODO: ReadFileToBuffer should return an error code or TResult
+		return AppError(fmt::format(L"Could not read file '{}'", FullPath.wstring()));
+	}
+	else
+	{
+		return ResultOk(std::move(Buffer));
+	}
+}
+
+TResult<FProxyDirectoryListing>
+FRemoteFileSystem::ListDirectory(const std::string_view RelativePath)
+{
+	FPooledHttpConnection HttpConnection(ProxyPool);
+	std::string			  FullPath;
+	FullPath.append(Root);
+	FullPath.append("/");
+	FullPath.append(RelativePath);
+	return ProxyQuery::ListDirectory(HttpConnection, ProxyPool.AuthDesc, FullPath);
+}
+
+TResult<FBuffer>
+FRemoteFileSystem::ReadFile(const std::string_view RelativePath)
+{
+	FPooledHttpConnection HttpConnection(ProxyPool);
+	std::string			  FullPath;
+	FullPath.append(Root);
+	FullPath.append("/");
+	FullPath.append(RelativePath);
+	return ProxyQuery::DownloadFile(HttpConnection, ProxyPool.AuthDesc, FullPath);
+}
+
 }  // namespace unsync
diff --git a/Engine/Source/Programs/Unsync/Private/UnsyncProxy.h b/Engine/Source/Programs/Unsync/Private/UnsyncProxy.h
--- a/Engine/Source/Programs/Unsync/Private/UnsyncProxy.h
+++ b/Engine/Source/Programs/Unsync/Private/UnsyncProxy.h
@@ -8,6 +8,8 @@
 #include "UnsyncSocket.h"
 #include "UnsyncThread.h"
 #include "UnsyncUtil.h"
+#include "UnsyncPool.h"
+#include "UnsyncHttp.h"
 
 #include <functional>
 #include <mutex>
@@ -163,6 +165,9 @@ public:
 	std::unique_ptr<FProxy> Alloc();
 	void					Dealloc(std::unique_ptr<FProxy>&& Proxy);
 
+	std::unique_ptr<FHttpConnection> AllocHttp();
+	void							 DeallocHttp(std::unique_ptr<FHttpConnection>&& Connection);
+
 	void Invalidate();
 	bool IsValid() const;
 
@@ -181,6 +186,8 @@ private:
 	std::vector<std::unique_ptr<FProxy>> Pool;
 	bool								 bValid = true;
 
+	std::optional<TObjectPool<FHttpConnection>> HttpPool;
+
 	FRemoteProtocolFeatures Features;
 	std::string SessionId;
 
@@ -189,6 +196,18 @@ private:
 	std::mutex Mutex;
 };
 
+struct FPooledHttpConnection
+{
+	FPooledHttpConnection(FProxyPool& InProxyPool) : ProxyPool(InProxyPool) { Inner = ProxyPool.AllocHttp(); }
+	~FPooledHttpConnection() { ProxyPool.DeallocHttp(std::move(Inner)); }
+	FHttpConnection* Get() { return Inner.get(); }
+	FHttpConnection& operator*() { return *Get(); }
+	FHttpConnection* operator->() { return Get(); }
+	operator FHttpConnection&() { return *Get(); }
+	FProxyPool&						 ProxyPool;
+	std::unique_ptr<FHttpConnection> Inner;
+};
+
 namespace ProxyQuery {
 
 struct FHelloResponse
@@ -219,7 +238,7 @@ TResult<FHelloResponse> Hello(FHttpConnection& Connection, const FAuthDesc* OptA
 
 struct FDirectoryListingEntry
 {
-	std::string Name;
+	std::string Name;  // utf-8
 	uint64		Mtime	   = 0;
 	uint64		Size	   = 0;
 	bool		bDirectory = false;
@@ -241,6 +260,40 @@ TResult<> DownloadFile(FHttpConnection&		   Connection,
 					   const std::string&	   Path,
 					   FDownloadOutputCallback OutputCallback);
 
-} 
+}
+
+using FProxyDirectoryListing = ProxyQuery::FDirectoryListing;
+using FProxyDirectoryEntry	 = ProxyQuery::FDirectoryListingEntry;
+
+// Abstracts basic filesystem operations, such as directory listing and file download.
+// Can be used to transparently handle basic local and remote file operations.
+struct FProxyFileSystem
+{
+	virtual TResult<FProxyDirectoryListing> ListDirectory(const std::string_view RelativePath) = 0;
+	virtual TResult<FBuffer>				ReadFile(const std::string_view RelativePath)	   = 0;
+
+	virtual ~FProxyFileSystem() = default;
+};
+
+struct FPhysicalFileSystem : public FProxyFileSystem
+{
+	FPhysicalFileSystem(const FPath& InRoot);
+
+	virtual TResult<FProxyDirectoryListing> ListDirectory(const std::string_view RelativePath) final override;
+	virtual TResult<FBuffer>				ReadFile(const std::string_view RelativePath) final override;
+
+	FPath Root;
+};
+
+struct FRemoteFileSystem : public FProxyFileSystem
+{
+	FRemoteFileSystem(const std::string& InRoot, FProxyPool& InProxyPool) : Root(InRoot), ProxyPool(InProxyPool) {}
+
+	virtual TResult<FProxyDirectoryListing> ListDirectory(const std::string_view RelativePath) final override;
+	virtual TResult<FBuffer>				ReadFile(const std::string_view RelativePath) final override;
+
+	std::string	Root;
+	FProxyPool& ProxyPool;
+};
 
 }  // namespace unsync
diff --git a/Engine/Source/Programs/Unsync/Private/UnsyncSerialization.cpp b/Engine/Source/Programs/Unsync/Private/UnsyncSerialization.cpp
--- a/Engine/Source/Programs/Unsync/Private/UnsyncSerialization.cpp
+++ b/Engine/Source/Programs/Unsync/Private/UnsyncSerialization.cpp
@@ -429,22 +429,6 @@ LoadPackReferences(FIOReaderStream& Reader, FSerializedSectionHeader Header)
 	return Result;
 }
 
-static void
-ConvertSerializedPathSeparatorsToNative(std::string& PathUtf8)
-{
-	std::replace_if(
-		PathUtf8.begin(),
-		PathUtf8.end(),
-		[](char C) { return C == '/' || C == '\\'; },
-		char(FPath::preferred_separator));
-}
-
-static void
-ConvertNativePathSeparatorsToSerialized(std::string& PathUtf8)
-{
-	std::replace(PathUtf8.begin(), PathUtf8.end(), char(FPath::preferred_separator), '/');
-}
-
 bool  // TODO: return a TResult
 LoadDirectoryManifest(FDirectoryManifest& OutManifest, const FPath& Root, FIOReaderStream& Stream)
 {
@@ -582,7 +566,7 @@ LoadDirectoryManifest(FDirectoryManifest& OutManifest, const FPath& Root, FIORea
 		for (uint64 FileIndex = 0; FileIndex < NumFiles; ++FileIndex)
 		{
 			Serialize(Stream, FilenameUtf8);
-			ConvertSerializedPathSeparatorsToNative(FilenameUtf8);
+			ConvertDirectorySeparatorsToNative(FilenameUtf8);
 			
 			FFileManifest FileManifest;
 			Serialize(Stream, FileManifest.Mtime);
@@ -809,7 +793,7 @@ SaveDirectoryManifest(const FDirectoryManifest& Manifest, FVectorStreamOut& Stre
 	for (const auto& It : Manifest.Files)
 	{
 		std::string FilenameUtf8 = ConvertWideToUtf8(It.first);
-		ConvertNativePathSeparatorsToSerialized(FilenameUtf8);
+		ConvertDirectorySeparatorsToUnix(FilenameUtf8);
 
 		Serialize(Stream, FilenameUtf8);  // name as utf8
 
diff --git a/Engine/Source/Programs/Unsync/Private/UnsyncSource.h b/Engine/Source/Programs/Unsync/Private/UnsyncSource.h
--- a/Engine/Source/Programs/Unsync/Private/UnsyncSource.h
+++ b/Engine/Source/Programs/Unsync/Private/UnsyncSource.h
@@ -3,6 +3,7 @@
 #pragma once
 
 #include "UnsyncCommon.h"
+#include "UnsyncProxy.h"
 
 namespace unsync {
 
diff --git a/Engine/Source/Programs/Unsync/Private/UnsyncUtil.cpp b/Engine/Source/Programs/Unsync/Private/UnsyncUtil.cpp
--- a/Engine/Source/Programs/Unsync/Private/UnsyncUtil.cpp
+++ b/Engine/Source/Programs/Unsync/Private/UnsyncUtil.cpp
@@ -195,6 +195,13 @@ FFileAttributeCache::Exists(const FPath& Path) const
 	return It != Map.end();
 }
 
+std::string
+ToString(const FPath& Path)
+{
+	FPathStringView PathView = Path.native();
+	return ToString(PathView);
+}
+
 std::wstring
 StringToLower(const std::wstring& Input)
 {
diff --git a/Engine/Source/Programs/Unsync/Private/UnsyncUtil.h b/Engine/Source/Programs/Unsync/Private/UnsyncUtil.h
--- a/Engine/Source/Programs/Unsync/Private/UnsyncUtil.h
+++ b/Engine/Source/Programs/Unsync/Private/UnsyncUtil.h
@@ -139,6 +139,20 @@ std::wstring ConvertUtf8ToWide(std::string_view StringUtf8);
 std::string	 ConvertWideToUtf8(std::wstring_view StringWide);
 void		 ConvertWideToUtf8(std::wstring_view StringWide, std::string& Result);
 
+inline std::string
+ToString(const std::wstring_view WideStringView)
+{
+	return ConvertWideToUtf8(WideStringView);
+}
+
+inline std::string
+ToString(const std::string_view StringView)
+{
+	return std::string(StringView);
+}
+
+std::string ToString(const FPath& Path);
+
 std::wstring StringToLower(const std::wstring& Input);
 std::wstring StringToUpper(const std::wstring& Input);
 
