unsync - Use pack files when downloading blocks from server when possible

diff --git a/Engine/Source/Programs/Unsync/Private/UnsyncCore.cpp b/Engine/Source/Programs/Unsync/Private/UnsyncCore.cpp
--- a/Engine/Source/Programs/Unsync/Private/UnsyncCore.cpp
+++ b/Engine/Source/Programs/Unsync/Private/UnsyncCore.cpp
@@ -21,6 +21,7 @@
 #include "UnsyncVersion.h"
 #include "UnsyncFilter.h"
 #include "UnsyncSource.h"
+#include "UnsyncPack.h"
 
 #include <condition_variable>
 #include <filesystem>
@@ -941,6 +942,7 @@ CopyFileIfPossiblyDifferent(FProxyFileSystem&	   FileSystem,
 
 static bool
 LoadAndMergeSourceManifest(FDirectoryManifest& Output,
+						   std::vector<FPackIndexDatabase>& OutIndexFiles,
 						   FProxyFileSystem&   ProxyFileSystem,
 						   const FPath&		   SourcePath,
 						   const FPath&		   TempPath,
@@ -989,17 +991,42 @@ LoadAndMergeSourceManifest(FDirectoryManifest& Output,
 		return false;
 	}
 
-	#if 0
 	if (!FindManifestResult->PackIndexFiles.empty())
 	{
 		UNSYNC_VERBOSE(L"Loading pack index database");
 		UNSYNC_LOG_INDENT;
-		for (const FRemoteFileInfo& FileInfo : FindManifestResult->PackIndexFiles)
+
+		std::unordered_set<FPath> FoundPackFiles;
+		for (const FRemoteFileInfo& PackFileInfo : FindManifestResult->PackDataFiles)
 		{
-			TResult<FBuffer> FileBuffer = ProxyFileSystem.ReadFile(ToString(FileInfo.Path));
+			FoundPackFiles.insert(PackFileInfo.Path);
+		}
+
+		for (const FRemoteFileInfo& IndexFileInfo : FindManifestResult->PackIndexFiles)
+		{
+			FPath PackDataFilePath = IndexFileInfo.Path;
+			PackDataFilePath.replace_extension(".unsync_pack");
+			if (!FoundPackFiles.contains(PackDataFilePath))
+			{
+				UNSYNC_WARNING(L"Could not find pack file '%ls'", PackDataFilePath.wstring().c_str());
+				continue;
+			}
+
+			UNSYNC_VERBOSE(L"Reading '%ls'", IndexFileInfo.Path.wstring().c_str());
+			TResult<FBuffer> FileBuffer = ProxyFileSystem.ReadFile(ToString(IndexFileInfo.Path));
+			FMemReader		 Reader(*FileBuffer);
+			FIOReaderStream	 Stream(Reader);
+
+			FPackIndexDatabase IndexFile;
+			IndexFile.IndexPath = SourcePath / IndexFileInfo.Path;
+			IndexFile.DataPath	= SourcePath / PackDataFilePath;
+
+			if (LoadPackIndexDatabase(IndexFile, Stream))
+			{
+				OutIndexFiles.emplace_back(std::move(IndexFile));
+			}
 		}
 	}
-	#endif
 
 	if (Output.IsValid() && !AlgorithmOptionsCompatible(Output.Algorithm, LoadedManifest.Algorithm))
 	{
@@ -1087,6 +1114,8 @@ struct FFileSyncTaskBatch
 			UNSYNC_UNUSED(DownloadResult);
 		}
 
+		ProxyPool.Dealloc(std::move(Proxy));
+
 		GScheduler->NetworkSemaphore.Release();
 
 		return Result;
@@ -1216,6 +1245,8 @@ SyncDirectory(const FSyncDirectoryOptions& SyncOptions)
 
 	FTimingLogger ManifestLoadTimingLogger("Manifest load time", ELogLevel::Info);
 
+	std::vector<FPackIndexDatabase> PackIndexFiles;
+
 	if (SyncOptions.SourceType == ESourceType::ServerWithManifestHash)
 	{
 		if (!ProxyPool.IsValid())
@@ -1247,6 +1278,7 @@ SyncDirectory(const FSyncDirectoryOptions& SyncOptions)
 	else if (!SyncOptions.SourceManifestOverride.empty())
 	{
 		bSourceManifestOk = LoadDirectoryManifest(SourceDirectoryManifest, SourcePath, SyncOptions.SourceManifestOverride);
+		// TODO: load pack index files
 		if (!bSourceManifestOk)
 		{
 			UNSYNC_ERROR(L"Could not load explicit manifest file");
@@ -1289,6 +1321,7 @@ SyncDirectory(const FSyncDirectoryOptions& SyncOptions)
 		for (const FPath& ThisSourcePath : AllSources)
 		{
 			if (!LoadAndMergeSourceManifest(SourceDirectoryManifest,
+											PackIndexFiles,
 											*ProxyFileSystem,
 											ThisSourcePath,
 											TargetTempPath,
@@ -1682,12 +1715,22 @@ SyncDirectory(const FSyncDirectoryOptions& SyncOptions)
 
 			if (bProxyHasData)
 			{
-				ProxyPool.InitRequestMap(SourceDirectoryManifest.Algorithm.StrongHashAlgorithmId);
+				FBlockRequestMap BlockRequestMap;
+				BlockRequestMap.Init(SourceDirectoryManifest.Algorithm.StrongHashAlgorithmId);
 
 				for (const FFileSyncTask& Task : AllFileTasks)
 				{
-					ProxyPool.BuildFileBlockRequests(Task.OriginalSourceFilePath, Task.ResolvedSourceFilePath, *Task.SourceManifest);
+					BlockRequestMap.AddFileBlocks(Task.OriginalSourceFilePath, Task.ResolvedSourceFilePath, *Task.SourceManifest);
 				}
+
+				// Override loose file blocks with pack files
+				for (const FPackIndexDatabase& Pack : PackIndexFiles)
+				{
+					FPath ResolvedDataPackPath = ResolvePath(Pack.DataPath);
+					BlockRequestMap.AddPackBlocks(Pack.DataPath, ResolvedDataPackPath, MakeView(Pack.Entries));
+				}
+
+				ProxyPool.SetRequestMap(std::move(BlockRequestMap));
 			}
 			else
 			{
@@ -1805,7 +1848,7 @@ SyncDirectory(const FSyncDirectoryOptions& SyncOptions)
 			UNSYNC_VERBOSE(L"Copy '%ls' (%ls)", Item.TargetFilePath.wstring().c_str(), (Item.NeedBytesFromBase) ? L"partial" : L"full");
 
 			std::unique_ptr<FNativeFile> BaseFile;
-			if (Item.IsBaseValid())
+			if (Item.IsBaseValid() && !Item.NeedList.Base.empty())
 			{
 				BaseFile = std::make_unique<FNativeFile>(Item.BaseFilePath, EFileMode::ReadOnlyUnbuffered);
 			}
diff --git a/Engine/Source/Programs/Unsync/Private/UnsyncFile.h b/Engine/Source/Programs/Unsync/Private/UnsyncFile.h
--- a/Engine/Source/Programs/Unsync/Private/UnsyncFile.h
+++ b/Engine/Source/Programs/Unsync/Private/UnsyncFile.h
@@ -427,6 +427,12 @@ struct FIOReaderStream
 
 	bool IsValid() const { return Inner.IsValid(); }
 
+	template <typename T>
+	inline uint64 ReadInto(T& Output)
+	{
+		return Read(&Output, sizeof(T));
+	}
+
 	FIOReader& Inner;
 	uint64	   Offset = 0;
 };
diff --git a/Engine/Source/Programs/Unsync/Private/UnsyncPack.cpp b/Engine/Source/Programs/Unsync/Private/UnsyncPack.cpp
--- a/Engine/Source/Programs/Unsync/Private/UnsyncPack.cpp
+++ b/Engine/Source/Programs/Unsync/Private/UnsyncPack.cpp
@@ -120,6 +120,8 @@ FPackWriteContext::InternalFinishPack()
 		uint64		 IndexDataSize = sizeof(IndexEntries[0]) * IndexEntries.size();
 
 		FPackIndexHeader IndexHeader;
+		IndexHeader.NumEntries = IndexEntries.size();
+
 		const uint64	 IndexTotalSize = sizeof(IndexHeader) + IndexDataSize;
 
 		FNativeFile IndexFile(FinalIndexFilename, EFileMode::CreateWriteOnly, IndexTotalSize);
@@ -251,4 +253,30 @@ void DeletePackAndIdexData(const FPath& PackRootDirectory)
 	}
 }
 
+bool LoadPackIndexDatabase(FPackIndexDatabase& Output, FIOReaderStream& Stream)
+{
+	FPackIndexHeader IndexHeader = {};
+	memset(&IndexHeader, 0, sizeof(IndexHeader));
+
+	Stream.ReadInto(IndexHeader);
+
+	if (IndexHeader.Version != FPackIndexHeader::VERSION || IndexHeader.Magic != FPackIndexHeader::MAGIC)
+	{
+		return false;
+	}
+
+	Output.Entries.resize(IndexHeader.NumEntries);
+
+	if (IndexHeader.NumEntries != 0)
+	{
+		const uint64 ExpectedReadSize = IndexHeader.NumEntries * sizeof(FPackIndexEntry);
+		const uint64 ReadSize = Stream.Read(Output.Entries.data(), ExpectedReadSize);
+		return ReadSize == ExpectedReadSize;
+	}
+	else
+	{
+		return true;
+	}
+}
+
 }  // namespace unsync
diff --git a/Engine/Source/Programs/Unsync/Private/UnsyncPack.h b/Engine/Source/Programs/Unsync/Private/UnsyncPack.h
--- a/Engine/Source/Programs/Unsync/Private/UnsyncPack.h
+++ b/Engine/Source/Programs/Unsync/Private/UnsyncPack.h
@@ -16,6 +16,17 @@ static constexpr uint64 GMaxPackFileSize = 1_GB;
 
 struct FDirectoryManifest;
 
+struct FPackIndexDatabase
+{
+	std::vector<FPackIndexEntry> Entries;
+
+	// Runtime-only data
+	FPath IndexPath;
+	FPath DataPath;
+};
+
+bool LoadPackIndexDatabase(FPackIndexDatabase& Output, FIOReaderStream& Stream);
+
 struct FPackWriteContext
 {
 	FPackWriteContext(const FPath& InOutputRoot);
diff --git a/Engine/Source/Programs/Unsync/Private/UnsyncProtocol.h b/Engine/Source/Programs/Unsync/Private/UnsyncProtocol.h
--- a/Engine/Source/Programs/Unsync/Private/UnsyncProtocol.h
+++ b/Engine/Source/Programs/Unsync/Private/UnsyncProtocol.h
@@ -227,12 +227,13 @@ static_assert(sizeof(FPackIndexEntry) == 40);
 struct FPackIndexHeader
 {
 	static constexpr uint64 MAGIC	= 0xEEC735E03053CC3Full;
-	static constexpr uint64 VERSION = 1;
+	static constexpr uint64 VERSION = 2;
 
-	uint64 Magic   = MAGIC;
-	uint64 Version = VERSION;
+	uint64 Magic	  = MAGIC;
+	uint64 Version	  = VERSION;
+	uint64 NumEntries = 0;
 };
-static_assert(sizeof(FPackIndexHeader) == 16);
+static_assert(sizeof(FPackIndexHeader) == 24);
 
 enum class EPackReferenceFlags : uint32 {
 	Default				= 0,
diff --git a/Engine/Source/Programs/Unsync/Private/UnsyncProxy.cpp b/Engine/Source/Programs/Unsync/Private/UnsyncProxy.cpp
--- a/Engine/Source/Programs/Unsync/Private/UnsyncProxy.cpp
+++ b/Engine/Source/Programs/Unsync/Private/UnsyncProxy.cpp
@@ -256,6 +256,16 @@ FUnsyncProtocolImpl::Download(const TArrayView<FNeedBlock> NeedBlocks, const FBl
 		}
 	}
 
+	std::sort(Requests.begin(), Requests.end(), [](const FBlockRequest& A, const FBlockRequest& B) -> bool
+	{
+		int32 FileCmp = std::memcmp(A.FilenameMd5.Data, B.FilenameMd5.Data, A.FilenameMd5.Size());
+		if (FileCmp != 0)
+		{
+			return FileCmp;
+		}
+		return A.Offset < B.Offset;
+	});
+
 	bool bOk = bIsConnetedToHost;
 
 	// begin the command
@@ -842,11 +852,9 @@ FUnsyncProtocolImpl::~FUnsyncProtocolImpl()
 	bIsConnetedToHost = false;
 }
 
-void
-FBlockRequestMap::AddFileBlocks(const FPath& OriginalFilePath, const FPath& ResolvedFilePath, const FFileManifest& FileManifest)
+FHash128
+FBlockRequestMap::AddFile(const FPath& OriginalFilePath, const FPath& ResolvedFilePath)
 {
-	UNSYNC_ASSERTF(StrongHasher != EStrongHashAlgorithmID::Invalid, L"Request map is not initialized");
-
 	std::string OriginalFilePathUtf8 = ConvertWideToUtf8(OriginalFilePath.wstring());
 	std::string ResolvedFilePathUtf8 = ConvertWideToUtf8(ResolvedFilePath.wstring());
 
@@ -861,10 +869,40 @@ FBlockRequestMap::AddFileBlocks(const FPath& OriginalFilePath, const FPath& Reso
 		FileListUtf8.push_back(OriginalFilePathUtf8);
 	}
 
+	return OriginalNameHash;
+}
+
+void
+FBlockRequestMap::AddPackBlocks(const FPath&					  OriginalFilePath,
+								const FPath&					  ResolvedFilePath,
+								const TArrayView<FPackIndexEntry> PackManifest)
+{
+	UNSYNC_ASSERTF(StrongHasher != EStrongHashAlgorithmID::Invalid, L"Request map is not initialized");
+
+	FHash128 FileId = AddFile(OriginalFilePath, ResolvedFilePath);
+
+	for (const FPackIndexEntry& Block : PackManifest)
+	{
+		FBlockRequest Request;
+		Request.FilenameMd5				 = FileId;
+		Request.BlockHash				 = Block.BlockHash;
+		Request.Offset					 = Block.PackBlockOffset;
+		Request.Size					 = Block.PackBlockSize;
+		BlockRequests[Request.BlockHash] = Request;
+	}
+}
+
+void
+FBlockRequestMap::AddFileBlocks(const FPath& OriginalFilePath, const FPath& ResolvedFilePath, const FFileManifest& FileManifest)
+{
+	UNSYNC_ASSERTF(StrongHasher != EStrongHashAlgorithmID::Invalid, L"Request map is not initialized");
+
+	FHash128 FileId = AddFile(OriginalFilePath, ResolvedFilePath);
+
 	for (const FGenericBlock& Block : FileManifest.Blocks)
 	{
 		FBlockRequest Request;
-		Request.FilenameMd5				 = OriginalNameHash;
+		Request.FilenameMd5				 = FileId;
 		Request.BlockHash				 = Block.HashStrong.ToHash128();  // #wip-widehash
 		Request.Offset					 = Block.Offset;
 		Request.Size					 = Block.Size;
@@ -1093,6 +1131,13 @@ FProxyPool::InitRequestMap(EStrongHashAlgorithmID InStrongHasher)
 	RequestMap.Init(InStrongHasher);
 }
 
+void
+FProxyPool::SetRequestMap(FBlockRequestMap&& InRequestMap)
+{
+	std::lock_guard<std::mutex> LockGuard(Mutex);
+	RequestMap = std::move(InRequestMap);
+}
+
 FPhysicalFileSystem::FPhysicalFileSystem(const FPath& InRoot) : Root(InRoot)
 {
 }
@@ -1145,8 +1190,12 @@ FRemoteFileSystem::ListDirectory(const std::string_view RelativePath)
 	FPooledHttpConnection HttpConnection(ProxyPool);
 	std::string			  FullPath;
 	FullPath.append(Root);
-	FullPath.append("/");
-	FullPath.append(RelativePath);
+	if (!RelativePath.empty())
+	{
+		FullPath.append("/");
+		FullPath.append(RelativePath);
+	}
+	ConvertDirectorySeparatorsToUnix(FullPath);
 	return ProxyQuery::ListDirectory(HttpConnection, ProxyPool.AuthDesc, FullPath);
 }
 
@@ -1158,6 +1207,7 @@ FRemoteFileSystem::ReadFile(const std::string_view RelativePath)
 	FullPath.append(Root);
 	FullPath.append("/");
 	FullPath.append(RelativePath);
+	ConvertDirectorySeparatorsToUnix(FullPath);
 	return ProxyQuery::DownloadFile(HttpConnection, ProxyPool.AuthDesc, FullPath);
 }
 
diff --git a/Engine/Source/Programs/Unsync/Private/UnsyncProxy.h b/Engine/Source/Programs/Unsync/Private/UnsyncProxy.h
--- a/Engine/Source/Programs/Unsync/Private/UnsyncProxy.h
+++ b/Engine/Source/Programs/Unsync/Private/UnsyncProxy.h
@@ -72,6 +72,7 @@ public:
 	}
 
 	void AddFileBlocks(const FPath& OriginalFilePath, const FPath& ResolvedFilePath, const FFileManifest& Manifest);
+	void AddPackBlocks(const FPath& OriginalFilePath, const FPath& ResolvedFilePath, const TArrayView<FPackIndexEntry> PackManifest);
 
 	const std::vector<std::string>& GetFileList() const { return FileListUtf8; }
 	const FBlockRequest*			FindRequest(const FGenericHash& BlockHash) const;
@@ -80,6 +81,9 @@ public:
 	FMacroBlockRequest				GetMacroBlockRequest(const FGenericHash& BlockHash) const;
 
 private:
+
+	FHash128 AddFile(const FPath& OriginalFilePath, const FPath& ResolvedFilePath);
+
 	EStrongHashAlgorithmID							 StrongHasher = EStrongHashAlgorithmID::Invalid;
 	std::vector<std::string>						 FileListUtf8;
 	std::unordered_map<FHash128, uint32>			 HashToFile;
@@ -175,6 +179,7 @@ public:
 	const FAuthDesc* AuthDesc = nullptr; // optional reference to externally-owned auth parameters
 
 	void InitRequestMap(EStrongHashAlgorithmID InStrongHasher);
+	void SetRequestMap(FBlockRequestMap&& InRequestMap);
 	void BuildFileBlockRequests(const FPath& OriginalFilePath, const FPath& ResolvedFilePath, const FFileManifest& FileManifest);
 
 	const FRemoteProtocolFeatures& GetFeatures() const { return Features; }
