unsync - Fix compile error on Linux

diff --git a/Engine/Source/Programs/Unsync/Private/UnsyncCore.cpp b/Engine/Source/Programs/Unsync/Private/UnsyncCore.cpp
--- a/Engine/Source/Programs/Unsync/Private/UnsyncCore.cpp
+++ b/Engine/Source/Programs/Unsync/Private/UnsyncCore.cpp
@@ -996,17 +996,17 @@ LoadAndMergeSourceManifest(FDirectoryManifest& Output,
 		UNSYNC_VERBOSE(L"Loading pack index database");
 		UNSYNC_LOG_INDENT;
 
-		std::unordered_set<FPath> FoundPackFiles;
+		std::unordered_set<FPathStringView> FoundPackFiles;
 		for (const FRemoteFileInfo& PackFileInfo : FindManifestResult->PackDataFiles)
 		{
-			FoundPackFiles.insert(PackFileInfo.Path);
+			FoundPackFiles.insert(PackFileInfo.Path.native());
 		}
 
 		for (const FRemoteFileInfo& IndexFileInfo : FindManifestResult->PackIndexFiles)
 		{
 			FPath PackDataFilePath = IndexFileInfo.Path;
 			PackDataFilePath.replace_extension(".unsync_pack");
-			if (!FoundPackFiles.contains(PackDataFilePath))
+			if (!FoundPackFiles.contains(PackDataFilePath.native()))
 			{
 				UNSYNC_WARNING(L"Could not find pack file '%ls'", PackDataFilePath.wstring().c_str());
 				continue;
