unsync - Detect and report errros when downloading pack index files

diff --git a/Engine/Source/Programs/Unsync/Private/UnsyncCore.cpp b/Engine/Source/Programs/Unsync/Private/UnsyncCore.cpp
--- a/Engine/Source/Programs/Unsync/Private/UnsyncCore.cpp
+++ b/Engine/Source/Programs/Unsync/Private/UnsyncCore.cpp
@@ -1014,6 +1014,12 @@ LoadAndMergeSourceManifest(FDirectoryManifest& Output,
 
 			UNSYNC_VERBOSE(L"Reading '%ls'", IndexFileInfo.Path.wstring().c_str());
 			TResult<FBuffer> FileBuffer = ProxyFileSystem.ReadFile(ToString(IndexFileInfo.Path));
+			if (FileBuffer.IsError())
+			{
+				LogError(FileBuffer.GetError());
+				return false;
+			}
+
 			FMemReader		 Reader(*FileBuffer);
 			FIOReaderStream	 Stream(Reader);
 
diff --git a/Engine/Source/Programs/Unsync/Private/UnsyncProxy.cpp b/Engine/Source/Programs/Unsync/Private/UnsyncProxy.cpp
--- a/Engine/Source/Programs/Unsync/Private/UnsyncProxy.cpp
+++ b/Engine/Source/Programs/Unsync/Private/UnsyncProxy.cpp
@@ -752,7 +752,9 @@ ProxyQuery::DownloadFile(FHttpConnection&					 InConnection,
 
 		if (Range.Size != Response.Buffer.Size())
 		{
-			Error.Set(AppError(L"Downloaded file chunk size mismatch"));
+			std::string ErrorMessage = fmt::format("Downloaded file chunk size mismatch. Expected {} bytes, got {} byte.", Range.Size, Response.Buffer.Size());
+			Error.Set(AppError(std::move(ErrorMessage)));
+			return;
 		}
 
 		uint64 WrittenBytes = Result.Write(Response.Buffer.Data(), Range.Offset, Range.Size);
