unsync - Add a command to decode binary manifest to Json

diff --git a/Engine/Source/Programs/Unsync/CMakeLists.txt b/Engine/Source/Programs/Unsync/CMakeLists.txt
--- a/Engine/Source/Programs/Unsync/CMakeLists.txt
+++ b/Engine/Source/Programs/Unsync/CMakeLists.txt
@@ -40,6 +40,8 @@ set(src
 	Private/UnsyncCmdDiff.h
 	Private/UnsyncCmdHash.cpp
 	Private/UnsyncCmdHash.h
+	Private/UnsyncCmdInfo.cpp
+	Private/UnsyncCmdInfo.h
 	Private/UnsyncCmdLogin.cpp
 	Private/UnsyncCmdLogin.h
 	Private/UnsyncCmdMount.cpp
diff --git a/Engine/Source/Programs/Unsync/Private/UnsyncCmdInfo.cpp b/Engine/Source/Programs/Unsync/Private/UnsyncCmdInfo.cpp
new file mode 100644
--- /dev/null
+++ b/Engine/Source/Programs/Unsync/Private/UnsyncCmdInfo.cpp
@@ -0,0 +1,255 @@
+// Copyright Epic Games, Inc. All Rights Reserved.
+
+#include "UnsyncCmdInfo.h"
+#include "UnsyncFile.h"
+#include "UnsyncManifest.h"
+#include "UnsyncFilter.h"
+#include "UnsyncSerialization.h"
+
+#include <fmt/format.h>
+#if __has_include(<fmt/xchar.h>)
+#	include <fmt/xchar.h>
+#endif
+
+namespace unsync {
+
+static void
+FilterManifest(const FSyncFilter& SyncFilter, FDirectoryManifest& Manifest)
+{
+	auto It = Manifest.Files.begin();
+	while (It != Manifest.Files.end())
+	{
+		if (SyncFilter.ShouldSync(It->first))
+		{
+			++It;
+		}
+		else
+		{
+			It = Manifest.Files.erase(It);
+		}
+	}
+}
+
+static THashMap<FGenericHash, FGenericBlock>
+BuildBlockMap(const FDirectoryManifest& Manifest, bool bNeedMacroBlocks)
+{
+	THashMap<FGenericHash, FGenericBlock> Result;
+	for (const auto& It : Manifest.Files)
+	{
+		const FFileManifest& File = It.second;
+		if (bNeedMacroBlocks)
+		{
+			for (const FGenericBlock& Block : File.MacroBlocks)
+			{
+				Result[Block.HashStrong] = Block;
+			}
+		}
+		else
+		{
+			for (const FGenericBlock& Block : File.Blocks)
+			{
+				Result[Block.HashStrong] = Block;
+			}
+		}
+	}
+	return Result;
+}
+
+
+void
+LogManifestDiff(ELogLevel LogLevel, const FDirectoryManifest& ManifestA, const FDirectoryManifest& ManifestB)
+{
+	THashMap<FGenericHash, FGenericBlock> BlocksA = BuildBlockMap(ManifestA, false);
+	THashMap<FGenericHash, FGenericBlock> BlocksB = BuildBlockMap(ManifestB, false);
+
+	THashMap<FGenericHash, FGenericBlock> MacroBlocksA = BuildBlockMap(ManifestA, true);
+	THashMap<FGenericHash, FGenericBlock> MacroBlocksB = BuildBlockMap(ManifestB, true);
+
+	uint32 NumCommonBlocks		= 0;
+	uint64 TotalCommonBlockSize = 0;
+	uint64 TotalSizeA			= 0;
+	uint64 TotalSizeB			= 0;
+
+	uint64 PatchSizeFromAtoB = 0;
+
+	for (const auto& ItA : BlocksA)
+	{
+		TotalSizeA += ItA.second.Size;
+		auto ItB = BlocksB.find(ItA.first);
+		if (ItB != BlocksB.end())
+		{
+			NumCommonBlocks++;
+			TotalCommonBlockSize += ItA.second.Size;
+		}
+	}
+
+	for (const auto& ItB : BlocksB)
+	{
+		TotalSizeB += ItB.second.Size;
+		if (BlocksA.find(ItB.first) == BlocksA.end())
+		{
+			PatchSizeFromAtoB += ItB.second.Size;
+		}
+	}
+
+	uint32 NumCommonMacroBlocks		 = 0;
+	uint64 TotalCommonMacroBlockSize = 0;
+	for (const auto& ItA : MacroBlocksA)
+	{
+		auto ItB = MacroBlocksB.find(ItA.first);
+		if (ItB != MacroBlocksB.end())
+		{
+			NumCommonMacroBlocks++;
+			TotalCommonMacroBlockSize += ItA.second.Size;
+		}
+	}
+
+	LogPrintf(LogLevel, L"Common macro blocks: %d, %.3f MB\n", NumCommonMacroBlocks, SizeMb(TotalCommonMacroBlockSize));
+
+	LogPrintf(LogLevel,
+			  L"Common blocks: %d, %.3f MB (%.2f%% of A, %.2f%% of B)\n",
+			  NumCommonBlocks,
+			  SizeMb(TotalCommonBlockSize),
+			  100.0 * double(TotalCommonBlockSize) / double(TotalSizeA),
+			  100.0 * double(TotalCommonBlockSize) / double(TotalSizeB));
+
+	LogPrintf(LogLevel, L"Patch size: %.3f MB\n", SizeMb(PatchSizeFromAtoB));
+}
+
+static void LogDecodedManifestJson(const FDirectoryManifest& Manifest)
+{
+	std::wstring Output;
+
+	Output += L"{\n"; // main object
+
+	const std::string_view StrongHash = ToString(Manifest.Algorithm.StrongHashAlgorithmId);
+	const std::string_view WeakHash	  = ToString(Manifest.Algorithm.WeakHashAlgorithmId);
+	const std::string_view Chunking	  = ToString(Manifest.Algorithm.ChunkingAlgorithmId);
+
+	FormatJsonKeyValueStr(Output, L"type", L"unsync_manifest", L",\n");
+	FormatJsonKeyValueStr(Output, L"hash_strong", ConvertUtf8ToWide(StrongHash), L",\n");
+	FormatJsonKeyValueStr(Output, L"hash_weak", ConvertUtf8ToWide(WeakHash), L",\n");
+	FormatJsonKeyValueStr(Output, L"chunking", ConvertUtf8ToWide(Chunking), L",\n");
+
+	Output += fmt::format(LR"("files": [)");
+	Output += L"\n";
+
+	uint64 FileIndex = 0;
+	for (const auto& FileIt : Manifest.Files)
+	{
+		const std::wstring& Name = FileIt.first;
+		const FFileManifest FileManifest = FileIt.second;
+
+		if (FileIndex != 0)
+		{
+			Output += L",\n";
+		}
+
+		Output += L"{";
+		FormatJsonKeyValueStr(Output, L"name", Name, L",");
+		FormatJsonKeyValueBool(Output, L"read_only", FileManifest.bReadOnly, L", ");
+		FormatJsonKeyValueUInt(Output, L"mtime", FileManifest.Mtime, L", ");
+		FormatJsonKeyValueUInt(Output, L"size", FileManifest.Size, L", ");
+		FormatJsonKeyValueUInt(Output, L"block_size", FileManifest.BlockSize, L", ");
+		FormatJsonKeyValueUInt(Output, L"num_blocks", FileManifest.Blocks.size(), L", ");
+		FormatJsonKeyValueUInt(Output, L"num_macro_blocks", FileManifest.MacroBlocks.size(), L",");
+		Output += L"\"blocks\": ";
+		FormatJsonBlockArray(Output, FileManifest.Blocks);
+		Output += L",";
+		Output += L"\"macro_blocks\": ";
+		FormatJsonBlockArray(Output, FileManifest.MacroBlocks);
+		Output += L"}";
+		++FileIndex;
+	}
+
+	Output += L"]\n"; // files
+	Output += L"}\n";  // main object
+
+	LogPrintf(ELogLevel::MachineReadable, Output.c_str());
+}
+
+int32
+CmdInfo(const FCmdInfoOptions& Options)
+{
+	FPath DirectoryManifestPathA = IsDirectory(Options.InputA) ? (Options.InputA / ".unsync" / "manifest.bin") : Options.InputA;
+	FPath DirectoryManifestPathB = IsDirectory(Options.InputB) ? (Options.InputB / ".unsync" / "manifest.bin") : Options.InputB;
+
+	FDirectoryManifest ManifestA;
+
+	bool bManifestAValid = LoadDirectoryManifest(ManifestA, Options.InputA, DirectoryManifestPathA);
+
+	if (!bManifestAValid)
+	{
+		return 1;
+	}
+
+	LogPrintf(ELogLevel::Info, L"Manifest A: %ls\n", DirectoryManifestPathA.wstring().c_str());
+
+	if (Options.SyncFilter)
+	{
+		FilterManifest(*Options.SyncFilter, ManifestA);
+	}
+
+	{
+		UNSYNC_LOG_INDENT;
+		LogManifestInfo(ELogLevel::Info, ManifestA);
+	}
+
+	if (Options.bListFiles)
+	{
+		UNSYNC_LOG_INDENT;
+		LogManifestFiles(ELogLevel::Info, ManifestA);
+	}
+
+	if (Options.bDecode)
+	{
+		LogDecodedManifestJson(ManifestA);
+		return 0;
+	}
+
+
+	if (Options.InputB.empty())
+	{
+		return 0;
+	}
+
+	LogPrintf(ELogLevel::Info, L"\n");
+
+	FDirectoryManifest ManifestB;
+
+	bool bManifestBValid = LoadDirectoryManifest(ManifestB, Options.InputB, DirectoryManifestPathB);
+
+	if (!bManifestBValid)
+	{
+		return 1;
+	}
+
+	LogPrintf(ELogLevel::Info, L"Manifest B: %ls\n", DirectoryManifestPathB.wstring().c_str());
+
+	if (Options.SyncFilter)
+	{
+		FilterManifest(*Options.SyncFilter, ManifestB);
+	}
+
+	{
+		UNSYNC_LOG_INDENT;
+		LogManifestInfo(ELogLevel::Info, ManifestB);
+	}
+	if (Options.bListFiles)
+	{
+		UNSYNC_LOG_INDENT;
+		LogManifestFiles(ELogLevel::Info, ManifestB);
+	}
+
+	LogPrintf(ELogLevel::Info, L"\n");
+	LogPrintf(ELogLevel::Info, L"Difference:\n");
+
+	{
+		UNSYNC_LOG_INDENT;
+		LogManifestDiff(ELogLevel::Info, ManifestA, ManifestB);
+	}
+
+	return 0;
+}
+
+}  // namespace unsync
diff --git a/Engine/Source/Programs/Unsync/Private/UnsyncCmdInfo.h b/Engine/Source/Programs/Unsync/Private/UnsyncCmdInfo.h
new file mode 100644
--- /dev/null
+++ b/Engine/Source/Programs/Unsync/Private/UnsyncCmdInfo.h
@@ -0,0 +1,22 @@
+// Copyright Epic Games, Inc. All Rights Reserved.
+
+#pragma once
+
+#include "UnsyncCommon.h"
+
+namespace unsync {
+
+struct FSyncFilter;
+
+struct FCmdInfoOptions
+{
+	FPath			   InputA;
+	FPath			   InputB;
+	bool			   bListFiles = false;
+	const FSyncFilter* SyncFilter = nullptr;
+	bool			   bDecode = false;
+};
+
+int32 CmdInfo(const FCmdInfoOptions& Options);
+
+}  // namespace unsync
diff --git a/Engine/Source/Programs/Unsync/Private/UnsyncMain.cpp b/Engine/Source/Programs/Unsync/Private/UnsyncMain.cpp
--- a/Engine/Source/Programs/Unsync/Private/UnsyncMain.cpp
+++ b/Engine/Source/Programs/Unsync/Private/UnsyncMain.cpp
@@ -10,6 +10,7 @@
 #include "UnsyncCmdPush.h"
 #include "UnsyncCmdQuery.h"
 #include "UnsyncCmdSync.h"
+#include "UnsyncCmdInfo.h"
 #include "UnsyncCore.h"
 #include "UnsyncFile.h"
 #include "UnsyncMemory.h"
@@ -109,7 +110,7 @@ InnerMain(int Argc, char** Argv)
 	bool					 bNoOutputFiles  	 = false;
 	bool					 bNoOutputRevisions  = false;
 	bool					 bPackOnlySmallFiles = false;
-	bool					 bPackFiles		 = false;
+	bool					 bPackFiles		     = false;
 	int32					 CompressionLevel	 = 3;
 	uint32					 DiffBlockSize		 = uint32(4_KB);
 	uint32					 HashOrSyncBlockSize = uint32(64_KB);
@@ -245,6 +246,7 @@ InnerMain(int Argc, char** Argv)
 	SubInfo->add_option("--exclude",
 						ExcludeFilterArrayUtf8,
 						"Exclude filenames that contain specified words (comma separated). Filter is run after --include.");
+	SubInfo->add_flag("--decode", bDecode, "Decode binary manifest into json");
 	SubCommands.push_back(SubInfo);
 
 	// Configure diff
@@ -431,6 +433,10 @@ InnerMain(int Argc, char** Argv)
 	{
 		GLogMachineReadable = true;
 	}
+	else if (Cli.got_subcommand(SubInfo))
+	{
+		GLogMachineReadable = bDecode;
+	}
 
 	UNSYNC_VERBOSE(L"UNSYNC v%hs", GetVersionString().c_str());
 
@@ -1049,6 +1055,7 @@ InnerMain(int Argc, char** Argv)
 		Options.InputB	   = InputFilename2;
 		Options.bListFiles = bInfoFiles;
 		Options.SyncFilter = &SyncFilter;
+		Options.bDecode	   = bDecode;
 		return CmdInfo(Options);
 	}
 	else if (Cli.got_subcommand(SubQuery))
diff --git a/Engine/Source/Programs/Unsync/Private/UnsyncManifest.cpp b/Engine/Source/Programs/Unsync/Private/UnsyncManifest.cpp
--- a/Engine/Source/Programs/Unsync/Private/UnsyncManifest.cpp
+++ b/Engine/Source/Programs/Unsync/Private/UnsyncManifest.cpp
@@ -557,185 +557,6 @@ LoadOrCreateDirectoryManifest(FDirectoryManifest& Result, const FPath& Root, con
 	return true;
 }
 
-THashMap<FGenericHash, FGenericBlock>
-BuildBlockMap(const FDirectoryManifest& Manifest, bool bNeedMacroBlocks)
-{
-	THashMap<FGenericHash, FGenericBlock> Result;
-	for (const auto& It : Manifest.Files)
-	{
-		const FFileManifest& File = It.second;
-		if (bNeedMacroBlocks)
-		{
-			for (const FGenericBlock& Block : File.MacroBlocks)
-			{
-				Result[Block.HashStrong] = Block;
-			}
-		}
-		else
-		{
-			for (const FGenericBlock& Block : File.Blocks)
-			{
-				Result[Block.HashStrong] = Block;
-			}
-		}
-	}
-	return Result;
-}
-
-void
-LogManifestDiff(ELogLevel LogLevel, const FDirectoryManifest& ManifestA, const FDirectoryManifest& ManifestB)
-{
-	THashMap<FGenericHash, FGenericBlock> BlocksA = BuildBlockMap(ManifestA, false);
-	THashMap<FGenericHash, FGenericBlock> BlocksB = BuildBlockMap(ManifestB, false);
-
-	THashMap<FGenericHash, FGenericBlock> MacroBlocksA = BuildBlockMap(ManifestA, true);
-	THashMap<FGenericHash, FGenericBlock> MacroBlocksB = BuildBlockMap(ManifestB, true);
-
-	uint32 NumCommonBlocks		= 0;
-	uint64 TotalCommonBlockSize = 0;
-	uint64 TotalSizeA			= 0;
-	uint64 TotalSizeB			= 0;
-
-	uint64 PatchSizeFromAtoB = 0;
-
-	for (const auto& ItA : BlocksA)
-	{
-		TotalSizeA += ItA.second.Size;
-		auto ItB = BlocksB.find(ItA.first);
-		if (ItB != BlocksB.end())
-		{
-			NumCommonBlocks++;
-			TotalCommonBlockSize += ItA.second.Size;
-		}
-	}
-
-	for (const auto& ItB : BlocksB)
-	{
-		TotalSizeB += ItB.second.Size;
-		if (BlocksA.find(ItB.first) == BlocksA.end())
-		{
-			PatchSizeFromAtoB += ItB.second.Size;
-		}
-	}
-
-	uint32 NumCommonMacroBlocks		 = 0;
-	uint64 TotalCommonMacroBlockSize = 0;
-	for (const auto& ItA : MacroBlocksA)
-	{
-		auto ItB = MacroBlocksB.find(ItA.first);
-		if (ItB != MacroBlocksB.end())
-		{
-			NumCommonMacroBlocks++;
-			TotalCommonMacroBlockSize += ItA.second.Size;
-		}
-	}
-
-	LogPrintf(LogLevel, L"Common macro blocks: %d, %.3f MB\n", NumCommonMacroBlocks, SizeMb(TotalCommonMacroBlockSize));
-
-	LogPrintf(LogLevel,
-			  L"Common blocks: %d, %.3f MB (%.2f%% of A, %.2f%% of B)\n",
-			  NumCommonBlocks,
-			  SizeMb(TotalCommonBlockSize),
-			  100.0 * double(TotalCommonBlockSize) / double(TotalSizeA),
-			  100.0 * double(TotalCommonBlockSize) / double(TotalSizeB));
-
-	LogPrintf(LogLevel, L"Patch size: %.3f MB\n", SizeMb(PatchSizeFromAtoB));
-}
-
-void
-FilterManifest(const FSyncFilter& SyncFilter, FDirectoryManifest& Manifest)
-{
-	auto It = Manifest.Files.begin();
-	while (It != Manifest.Files.end())
-	{
-		if (SyncFilter.ShouldSync(It->first))
-		{
-			++It;
-		}
-		else
-		{
-			It = Manifest.Files.erase(It);
-		}
-	}
-}
-
-int32
-CmdInfo(const FCmdInfoOptions& Options)
-{
-	FPath DirectoryManifestPathA = IsDirectory(Options.InputA) ? (Options.InputA / ".unsync" / "manifest.bin") : Options.InputA;
-	FPath DirectoryManifestPathB = IsDirectory(Options.InputB) ? (Options.InputB / ".unsync" / "manifest.bin") : Options.InputB;
-
-	FDirectoryManifest ManifestA;
-
-	bool bManifestAValid = LoadDirectoryManifest(ManifestA, Options.InputA, DirectoryManifestPathA);
-
-	if (!bManifestAValid)
-	{
-		return 1;
-	}
-
-	LogPrintf(ELogLevel::Info, L"Manifest A: %ls\n", DirectoryManifestPathA.wstring().c_str());
-
-	if (Options.SyncFilter)
-	{
-		FilterManifest(*Options.SyncFilter, ManifestA);
-	}
-
-	{
-		UNSYNC_LOG_INDENT;
-		LogManifestInfo(ELogLevel::Info, ManifestA);
-	}
-
-	if (Options.bListFiles)
-	{
-		UNSYNC_LOG_INDENT;
-		LogManifestFiles(ELogLevel::Info, ManifestA);
-	}
-
-	if (Options.InputB.empty())
-	{
-		return 0;
-	}
-
-	LogPrintf(ELogLevel::Info, L"\n");
-
-	FDirectoryManifest ManifestB;
-
-	bool bManifestBValid = LoadDirectoryManifest(ManifestB, Options.InputB, DirectoryManifestPathB);
-
-	if (!bManifestBValid)
-	{
-		return 1;
-	}
-
-	LogPrintf(ELogLevel::Info, L"Manifest B: %ls\n", DirectoryManifestPathB.wstring().c_str());
-
-	if (Options.SyncFilter)
-	{
-		FilterManifest(*Options.SyncFilter, ManifestB);
-	}
-
-	{
-		UNSYNC_LOG_INDENT;
-		LogManifestInfo(ELogLevel::Info, ManifestB);
-	}
-	if (Options.bListFiles)
-	{
-		UNSYNC_LOG_INDENT;
-		LogManifestFiles(ELogLevel::Info, ManifestB);
-	}
-
-	LogPrintf(ELogLevel::Info, L"\n");
-	LogPrintf(ELogLevel::Info, L"Difference:\n");
-
-	{
-		UNSYNC_LOG_INDENT;
-		LogManifestDiff(ELogLevel::Info, ManifestA, ManifestB);
-	}
-
-	return 0;
-}
-
 FHash256
 ComputeSerializedManifestHash(const FDirectoryManifest& Manifest)
 {
diff --git a/Engine/Source/Programs/Unsync/Private/UnsyncManifest.h b/Engine/Source/Programs/Unsync/Private/UnsyncManifest.h
--- a/Engine/Source/Programs/Unsync/Private/UnsyncManifest.h
+++ b/Engine/Source/Programs/Unsync/Private/UnsyncManifest.h
@@ -111,13 +111,4 @@ bool			   MergeManifests(FDirectoryManifest& Existing, const FDirectoryManifest&
 void			   MoveCompatibleManifestBlocks(FDirectoryManifest& Manifest, FDirectoryManifest&& DonorManifest);
 bool			   AlgorithmOptionsCompatible(const FAlgorithmOptions& A, const FAlgorithmOptions& B);
 
-struct FCmdInfoOptions
-{
-	FPath			   InputA;
-	FPath			   InputB;
-	bool			   bListFiles = false;
-	const FSyncFilter* SyncFilter = nullptr;
-};
-int32 CmdInfo(const FCmdInfoOptions& Options);
-
 }  // namespace unsync
diff --git a/Engine/Source/Programs/Unsync/Private/UnsyncProtocol.h b/Engine/Source/Programs/Unsync/Private/UnsyncProtocol.h
--- a/Engine/Source/Programs/Unsync/Private/UnsyncProtocol.h
+++ b/Engine/Source/Programs/Unsync/Private/UnsyncProtocol.h
@@ -73,6 +73,10 @@ using FBlock256 = TBlock<FHash256>;
 using FGenericBlock = TBlock<FGenericHash>;
 using FGenericBlockArray = std::vector<FGenericBlock>;
 
+// Json formatting helpers
+void FormatJsonBlock(std::wstring& Output, const FGenericBlock& Block);
+void FormatJsonBlockArray(std::wstring& Output, const FGenericBlockArray& Blocks);
+
 static constexpr uint64 SERIALIZED_SECTION_ID_TERMINATOR			= 0;
 static constexpr uint64 SERIALIZED_SECTION_ID_METADATA_STRING		= 0xC6BD6CDCEEF79533ull;
 static constexpr uint64 SERIALIZED_SECTION_ID_MACRO_BLOCK			= 0x8390AEBB745E08BCull;
diff --git a/Engine/Source/Programs/Unsync/Private/UnsyncUtil.cpp b/Engine/Source/Programs/Unsync/Private/UnsyncUtil.cpp
--- a/Engine/Source/Programs/Unsync/Private/UnsyncUtil.cpp
+++ b/Engine/Source/Programs/Unsync/Private/UnsyncUtil.cpp
@@ -3,6 +3,7 @@
 #include "UnsyncUtil.h"
 #include "UnsyncFile.h"
 #include "UnsyncSocket.h"
+#include "UnsyncProtocol.h"
 
 #if UNSYNC_PLATFORM_WINDOWS
 #	include <Windows.h>
@@ -27,6 +28,9 @@
 #include <unordered_set>
 #include <system_error>
 #include <fmt/format.h>
+#if __has_include(<fmt/xchar.h>)
+#	include <fmt/xchar.h>
+#endif
 
 namespace unsync {
 
@@ -34,19 +38,32 @@ static FBuffer GSystemRootCerts;
 
 static const char G_HEX_CHARS[] = "0123456789abcdef";
 
+template<typename CharT>
 uint64
-BytesToHexChars(char* Output, uint64 OutputSize, const uint8* Input, uint64 InputSize)
+BytesToHexCharsT(CharT* Output, uint64 OutputSize, const uint8* Input, uint64 InputSize)
 {
 	const uint64 MaxBytes = std::min(OutputSize / 2, InputSize);
 	for (uint64 I = 0; I < MaxBytes; ++I)
 	{
 		uint8 V			  = Input[I];
-		Output[I * 2 + 0] = G_HEX_CHARS[V >> 4];
-		Output[I * 2 + 1] = G_HEX_CHARS[V & 0xF];
+		Output[I * 2 + 0] = CharT(G_HEX_CHARS[V >> 4]);
+		Output[I * 2 + 1] = CharT(G_HEX_CHARS[V & 0xF]);
 	}
 	return MaxBytes * 2;
 }
 
+uint64
+BytesToHexChars(char* Output, uint64 OutputSize, const uint8* Input, uint64 InputSize)
+{
+	return BytesToHexCharsT(Output, OutputSize, Input, InputSize);
+}
+
+uint64
+BytesToHexChars(wchar_t* Output, uint64 OutputSize, const uint8* Input, uint64 InputSize)
+{
+	return BytesToHexCharsT(Output, OutputSize, Input, InputSize);
+}
+
 std::string
 BytesToHexString(const uint8* Data, uint64 Size)
 {
@@ -63,6 +80,67 @@ BytesToHexString(const uint8* Data, uint64 Size)
 	return Result;
 }
 
+
+void
+FormatJsonKeyValueStr(std::wstring& Output, std::wstring_view K, std::wstring_view V, std::wstring_view Suffix)
+{
+	fmt::format_to(std::back_inserter(Output), L"\"{}\": \"{}\"{}", K, V, Suffix);
+}
+
+void
+FormatJsonKeyValueUInt(std::wstring& Output, std::wstring_view K, uint64 V, std::wstring_view Suffix)
+{
+	fmt::format_to(std::back_inserter(Output), L"\"{}\": {}{}", K, V, Suffix);
+}
+
+void
+FormatJsonKeyValueBool(std::wstring& Output, std::wstring_view K, bool V, std::wstring_view Suffix)
+{
+	fmt::format_to(std::back_inserter(Output), L"\"{}\": {}{}", K, V ? L"true" : L"false", Suffix);
+}
+
+void
+FormatJsonBlock(std::wstring& Output, const FGenericBlock& Block)
+{
+	Output += L"{";
+
+	static const size_t MaxHashLen = 2 * sizeof(Block.HashStrong.Data);
+	wchar_t				HashChars[MaxHashLen];
+
+	uint64			  HashLen = BytesToHexChars(HashChars, MaxHashLen, Block.HashStrong.Data, Block.HashStrong.Size());
+	std::wstring_view HashStr = std::wstring_view(HashChars, HashLen);
+
+	FormatJsonKeyValueUInt(Output, L"offset", Block.Offset, L", ");
+	FormatJsonKeyValueUInt(Output, L"size", Block.Size, L", ");
+	if (Block.HashWeak != 0)
+	{
+		FormatJsonKeyValueUInt(Output, L"hash_weak", Block.HashWeak, L", ");
+	}
+	FormatJsonKeyValueStr(Output, L"hash_strong", HashStr);
+
+	Output += L"}";
+}
+
+void
+FormatJsonBlockArray(std::wstring& Output, const FGenericBlockArray& Blocks)
+{
+	Output += L"[\n";
+	uint64 BlockIndex = 0;
+	for (const FGenericBlock& Block : Blocks)
+	{
+		if (BlockIndex != 0)
+		{
+			Output += L",\n";
+		}
+
+		FormatJsonBlock(Output, Block);
+
+		++BlockIndex;
+	}
+	Output += L"]";
+}
+
+
 FTimingLogger::FTimingLogger(const char* InName, ELogLevel InLogLevel, bool bInEnabled)
 : bEnabled(bInEnabled)
 , Name(InName)
diff --git a/Engine/Source/Programs/Unsync/Private/UnsyncUtil.h b/Engine/Source/Programs/Unsync/Private/UnsyncUtil.h
--- a/Engine/Source/Programs/Unsync/Private/UnsyncUtil.h
+++ b/Engine/Source/Programs/Unsync/Private/UnsyncUtil.h
@@ -127,6 +127,7 @@ std::string BytesToHexString(const uint8* Data, uint64 Size);
 
 // Converts input bytes to hexadecimal ACII. Returns how many characters were written to output.
 uint64 BytesToHexChars(char* Output, uint64 OutputSize, const uint8* Input, uint64 InputSize);
+uint64 BytesToHexChars(wchar_t* Output, uint64 OutputSize, const uint8* Input, uint64 InputSize);
 
 template<typename HashType>
 std::string
@@ -371,4 +372,9 @@ bool LooksLikeUrl(std::string_view Str);
 bool LooksLikeHash160(const std::string_view Str);
 bool LooksLikeHash160(const std::wstring_view Str);
 
+// Json formatting helpers
+void FormatJsonKeyValueStr(std::wstring& Output, std::wstring_view K, std::wstring_view V, std::wstring_view Suffix = {});
+void FormatJsonKeyValueUInt(std::wstring& Output, std::wstring_view K, uint64 V, std::wstring_view Suffix = {});
+void FormatJsonKeyValueBool(std::wstring& Output, std::wstring_view K, bool V, std::wstring_view Suffix = {});
+
 }  // namespace unsync
