unsync - Add experimental support for pure HTTP proxy server API, instead of custom binary protocol

diff --git a/Engine/Source/Programs/Unsync/Private/UnsyncCore.cpp b/Engine/Source/Programs/Unsync/Private/UnsyncCore.cpp
--- a/Engine/Source/Programs/Unsync/Private/UnsyncCore.cpp
+++ b/Engine/Source/Programs/Unsync/Private/UnsyncCore.cpp
@@ -39,6 +39,7 @@ UNSYNC_THIRD_PARTY_INCLUDES_END
 namespace unsync {
 
 bool GDryRun = false;
+bool GExperimental = false;
 
 FBlock128
 ToBlock128(const FGenericBlock& GenericBlock)
diff --git a/Engine/Source/Programs/Unsync/Private/UnsyncCore.h b/Engine/Source/Programs/Unsync/Private/UnsyncCore.h
--- a/Engine/Source/Programs/Unsync/Private/UnsyncCore.h
+++ b/Engine/Source/Programs/Unsync/Private/UnsyncCore.h
@@ -21,6 +21,7 @@
 namespace unsync {
 
 extern bool GDryRun;
+extern bool GExperimental;
 
 class FProxy;
 class FProxyPool;
diff --git a/Engine/Source/Programs/Unsync/Private/UnsyncHttp.cpp b/Engine/Source/Programs/Unsync/Private/UnsyncHttp.cpp
--- a/Engine/Source/Programs/Unsync/Private/UnsyncHttp.cpp
+++ b/Engine/Source/Programs/Unsync/Private/UnsyncHttp.cpp
@@ -38,18 +38,18 @@ UpdateView(std::string_view& View, const char* Data, size_t Size)
 	}
 }
 
-using HttpMessageCallback = std::function<void(FHttpResponse&& Response)>;
-
 struct FHttpParser
 {
 	static FHttpParser* ToThis(http_parser* Parser) { return (FHttpParser*)(Parser->data); }
 
-	FHttpParser(HttpMessageCallback InResponseCallback,
+	FHttpParser(FHttpMessageCallback InResponseCallback,
+				FHttpChunkCallback   InChunkCallback,
 				uint8*				InScratchBuffer,
 				uint64				InScratchSize,
 				http_parser_type	Type,
 				EHttpMethod			InMethod)
 	: ResponseCallback(InResponseCallback)
+	, ChunkCallback(InChunkCallback)
 	, Method(InMethod)
 	, ScratchBuffer(InScratchBuffer)
 	, ScratchSize(InScratchSize)
@@ -95,9 +95,19 @@ struct FHttpParser
 		return 0;
 	}
 
-	int OnChunkHeader() { return 0; }
+	int OnChunkHeader()
+	{ 
+		return 0;
+	}
 
-	int OnChunkComplete() { return 0; }
+	int OnChunkComplete()
+	{
+		if (ChunkCallback)
+		{
+			ChunkCallback(Response);
+		}
+		return 0;
+	}
 
 	int OnHdrField(const char* Data, size_t Size)
 	{
@@ -214,7 +224,8 @@ struct FHttpParser
 		return bShouldContinue;
 	}
 
-	HttpMessageCallback ResponseCallback;
+	FHttpMessageCallback ResponseCallback;
+	FHttpChunkCallback ChunkCallback;
 	FHttpResponse		Response;
 
 	EHttpMethod Method = EHttpMethod::GET;
@@ -239,13 +250,13 @@ struct FHttpParser
 };
 
 FHttpResponse
-HttpRequest(FHttpConnection& Connection, const FHttpRequest& Request)
+HttpRequest(FHttpConnection& Connection, const FHttpRequest& Request, FHttpChunkCallback ChunkCallback)
 {
 	FHttpResponse Result;
 
 	if (HttpRequestBegin(Connection, Request))
 	{
-		Result = HttpRequestEnd(Connection);
+		Result = HttpRequestEnd(Connection, ChunkCallback);
 	}
 
 	return Result;
@@ -437,7 +448,7 @@ HttpRequestBegin(FHttpConnection& Connection, const FHttpRequest& Request)
 }
 
 FHttpResponse
-HttpRequestEnd(FHttpConnection& Connection)
+HttpRequestEnd(FHttpConnection& Connection, FHttpChunkCallback ChunkCallback)
 {
 	FHttpResponse Result;
 
@@ -467,7 +478,7 @@ HttpRequestEnd(FHttpConnection& Connection)
 	// TODO: user-provided scratch buffer
 	uint8 ScratchBuffer[256_KB];
 	ScratchBuffer[0] = 0;
-	FHttpParser Parser(MessageCallback, ScratchBuffer, sizeof(ScratchBuffer), HTTP_RESPONSE, Connection.Method);
+	FHttpParser Parser(MessageCallback, ChunkCallback, ScratchBuffer, sizeof(ScratchBuffer), HTTP_RESPONSE, Connection.Method);
 
 	while (!Parser.bComplete)
 	{
diff --git a/Engine/Source/Programs/Unsync/Private/UnsyncHttp.h b/Engine/Source/Programs/Unsync/Private/UnsyncHttp.h
--- a/Engine/Source/Programs/Unsync/Private/UnsyncHttp.h
+++ b/Engine/Source/Programs/Unsync/Private/UnsyncHttp.h
@@ -10,10 +10,12 @@
 #include <deque>
 #include <memory>
 #include <string>
+#include <functional>
 
 namespace unsync {
 
 struct FRemoteDesc;
+struct FHttpResponse;
 
 enum class EHttpContentType
 {
@@ -34,6 +36,9 @@ enum class EHttpMethod
 	PUT,
 };
 
+using FHttpMessageCallback = std::function<void(FHttpResponse&& Response)>;
+using FHttpChunkCallback   = std::function<void(FHttpResponse& Response)>;
+
 struct FHttpRequest
 {
 	EHttpMethod		 Method		   = EHttpMethod::GET;
@@ -111,7 +116,7 @@ const char* HttpStatusToString(int32 Code);
 
 // Synchronous HTTP request API
 
-FHttpResponse HttpRequest(FHttpConnection& Connection, const FHttpRequest& Request);
+FHttpResponse HttpRequest(FHttpConnection& Connection, const FHttpRequest& Request, FHttpChunkCallback ChunkCallback = {});
 
 inline FHttpResponse
 HttpRequest(FHttpConnection& Connection,
@@ -149,8 +154,10 @@ HttpRequest(FHttpConnection& Connection,
 	return HttpRequest(Connection, Request);
 }
 
-FHttpResponse
-HttpRequest(const FRemoteDesc& RemoteDesc, EHttpMethod Method, std::string_view RequestUrl, std::string_view BearerToken = {});
+FHttpResponse HttpRequest(const FRemoteDesc& RemoteDesc,
+						  EHttpMethod		 Method,
+						  std::string_view	 RequestUrl,
+						  std::string_view	 BearerToken = {});
 
 FHttpResponse HttpRequest(const FRemoteDesc& RemoteDesc,
 						  EHttpMethod		 Method,
@@ -163,6 +170,6 @@ FHttpResponse HttpRequest(const FRemoteDesc& RemoteDesc,
 
 bool HttpRequestBegin(FHttpConnection& Connection, const FHttpRequest& Request);
 
-FHttpResponse HttpRequestEnd(FHttpConnection& Connection);
+FHttpResponse HttpRequestEnd(FHttpConnection& Connection, FHttpChunkCallback ChunkCallback = {});
 
 }  // namespace unsync
diff --git a/Engine/Source/Programs/Unsync/Private/UnsyncMain.cpp b/Engine/Source/Programs/Unsync/Private/UnsyncMain.cpp
--- a/Engine/Source/Programs/Unsync/Private/UnsyncMain.cpp
+++ b/Engine/Source/Programs/Unsync/Private/UnsyncMain.cpp
@@ -116,6 +116,10 @@ InnerMain(int Argc, char** Argv)
 	uint32					 HashOrSyncBlockSize = uint32(64_KB);
 	uint32					 BackgroundTaskMemoryBudgetGB = 2;
 
+	const std::string HiddenGroupId;  // CLI11 uses an empty string group name to mark arguments that should be hidden
+	const std::string ExperimentalGroupId = "Experimental";
+	const std::string DangerousGroupId	  = "Dangerous";
+
 	struct FDeprecatedOptions
 	{
 		bool bQuickSyncMode			= false;
@@ -123,12 +127,11 @@ InnerMain(int Argc, char** Argv)
 		bool bQuickSourceValidation = false;
 	} DeprecatedOptions;
 
-	const std::string HiddenGroupId;  // CLI11 uses an empty string group name to mark arguments that should be hidden
-
-	auto AddTlsOptions = [&CacertFilenameUtf8, &bRequireTls, &bAllowInsecureTls](CLI::App* App) {
+	auto AddTlsOptions = [&CacertFilenameUtf8, &bRequireTls, &bAllowInsecureTls, &DangerousGroupId](CLI::App* App)
+	{
 		App->add_option("--cacert", CacertFilenameUtf8, "Certificate authority file to use for TLS validation (.pem)");
 		App->add_flag("--tls", bRequireTls, "Force TLS when connecting to remote server");
-		App->add_flag("--insecure", bAllowInsecureTls, "Skip remote server TLS certificate validation");
+		App->add_flag("--insecure", bAllowInsecureTls, "Skip remote server TLS certificate validation")->group(DangerousGroupId);
 	};
 
 	auto AddProxyOptions = [&RemoteAddressUtf8, &bNoProxySelect](CLI::App* App) {
@@ -182,7 +185,7 @@ InnerMain(int Argc, char** Argv)
 	{
 		SubPack =
 			Cli.add_subcommand("pack",
-							   "EXPERIMENTAL: Generate manifest for a directory and store all referenced data in a compressed pack file");
+							   "EXPERIMENTAL: Generate manifest for a directory and store all referenced data in a compressed pack file")->group(HiddenGroupId);
 
 		SubPack->add_option("Input", InputFilenameUtf8, "Input directory path")->required();
 
@@ -203,7 +206,7 @@ InnerMain(int Argc, char** Argv)
 
 	CLI::App* SubUnpack = nullptr;
 	{
-		SubUnpack = Cli.add_subcommand("unpack", "EXPERIMENTAL: Sync directory based on package snapshot");
+		SubUnpack = Cli.add_subcommand("unpack", "EXPERIMENTAL: Sync directory based on package snapshot")->group(HiddenGroupId);
 
 		SubUnpack->add_option("Output", OutputFilenameUtf8, "Output directory path")->required();
 
@@ -270,7 +273,7 @@ InnerMain(int Argc, char** Argv)
 	SubSync->add_option("Target", TargetFilenameUtf8, "Target path")->required();
 	SubSync->add_option("-m, --manifest", SourceManifestFilenameUtf8, "Override manifest path for Source");
 	AddProxyOptions(SubSync);
-	SubSync->add_option("--dfs", PreferredDfsUtf8, "Preferred DFS mirror (matched by sub-string)");
+	SubSync->add_option("--dfs", PreferredDfsUtf8, "DEPRECATED: Preferred DFS mirror (matched by sub-string)")->group(HiddenGroupId);
 	SubSync->add_option(
 		"--overlay",
 		OverlayArrayUtf8,
@@ -325,9 +328,12 @@ InnerMain(int Argc, char** Argv)
 					  "the manifest. This is an extra precaution that will detect any missing or invalid remote files before running the "
 					  "sync process, however this can be very slow when dealing with large numbers of files and directories.");
 
-	SubSync->add_flag("--no-output-validation", bNoOutputValidation, "Skip final patched file block hash validation (DANGEROUS)");
-	SubSync->add_flag("--no-space-validation", bNoSpaceValidation, "Skip checking available disk space before sync (DANGEROUS)");
-	SubSync->add_option("--scavenge", ScavengeRootUtf8, "Search for unsync manifests and reusable blocks in this directory (EXPERIMENTAL)");
+	SubSync->add_flag("--no-output-validation", bNoOutputValidation, "Skip final patched file block hash validation (DANGEROUS)")
+		->group(DangerousGroupId);
+	SubSync->add_flag("--no-space-validation", bNoSpaceValidation, "Skip checking available disk space before sync (DANGEROUS)")
+		->group(DangerousGroupId);
+	SubSync->add_option("--scavenge", ScavengeRootUtf8, "Search for unsync manifests and reusable blocks in this directory (EXPERIMENTAL)")
+		->group(ExperimentalGroupId);
 	SubSync->add_flag("--login", bShouldLogin, "Use user authentication when accessing unsync server");
 	SubSync->add_option("--token", AuthTokenPathUtf8, "Explicit path to the authentication token file to use");
 	SubSync->add_flag("--no-timeout", bNoSocketTimeout, "Disable the default 60 second timeout on network socket operations");
@@ -374,7 +380,7 @@ InnerMain(int Argc, char** Argv)
 
 	// Configure mount
 
-	CLI::App* SubMount = Cli.add_subcommand("mount", "Mount directory manifest as a virtual file system (EXPERIMENTAL)");
+	CLI::App* SubMount = Cli.add_subcommand("mount", "Mount directory manifest as a virtual file system (EXPERIMENTAL)")->group(HiddenGroupId);
 	SubMount
 		->add_option("Source",
 					 SourceFilenameUtf8,
@@ -393,6 +399,7 @@ InnerMain(int Argc, char** Argv)
 		Subcommand->add_option("--threads", GMaxThreads, "Limit worker threads to specified number");
 		Subcommand->add_flag("--buffered-files", GForceBufferedFiles, "Always use buffered file IO");
 		Subcommand->add_flag("--debug", bUseDebugMode, "Enable extra debugging features, such as extra memory safety validation");
+		Subcommand->add_flag("--experimental", GExperimental, "Enable experimental code paths")->group(HiddenGroupId);
 
 		SilentFlag->excludes(VerboseFlag);
 		SilentFlag->excludes(VeryVerboseFlag);
@@ -650,6 +657,16 @@ InnerMain(int Argc, char** Argv)
 
 	FPath SourceFilename = bFilesystemSource ? NormalizeFilenameUtf8(SourceFilenameUtf8) : FPath(SourceFilenameUtf8);
 
+	if (GDryRun)
+	{
+		UNSYNC_LOG(L">>> DRY RUN <<<");
+	}
+
+	if (GExperimental)
+	{
+		UNSYNC_LOG(L">>> EXPERIMENTAL MODE <<<");
+	}
+
 	if (GLogVeryVerbose)
 	{
 		UNSYNC_LOG(L"Very verbose logging is enabled");
@@ -659,11 +676,6 @@ InnerMain(int Argc, char** Argv)
 		UNSYNC_LOG(L"Verbose logging is enabled");
 	}
 
-	if (GDryRun)
-	{
-		UNSYNC_LOG(L">>> DRY RUN <<<");
-	}
-
 	if (GForceBufferedFiles)
 	{
 		UNSYNC_VERBOSE(L"Using buffered file IO");
@@ -868,6 +880,8 @@ InnerMain(int Argc, char** Argv)
 		}
 	}
 
+	RemoteDesc.TlsRequirement = ETlsRequirement::None; // nocheckin
+
 	if (Cli.got_subcommand(SubHash))
 	{
 		FCmdHashOptions HashOptions;
diff --git a/Engine/Source/Programs/Unsync/Private/UnsyncProtocol.h b/Engine/Source/Programs/Unsync/Private/UnsyncProtocol.h
--- a/Engine/Source/Programs/Unsync/Private/UnsyncProtocol.h
+++ b/Engine/Source/Programs/Unsync/Private/UnsyncProtocol.h
@@ -75,7 +75,9 @@ using FGenericBlockArray = std::vector<FGenericBlock>;
 
 // Json formatting helpers
 void FormatJsonBlock(std::wstring& Output, const FGenericBlock& Block);
+void FormatJsonBlock(std::string& Output, const FGenericBlock& Block);
 void FormatJsonBlockArray(std::wstring& Output, const FGenericBlockArray& Blocks);
+void FormatJsonBlockArray(std::string& Output, const FGenericBlockArray& Blocks);
 
 static constexpr uint64 SERIALIZED_SECTION_ID_TERMINATOR			= 0;
 static constexpr uint64 SERIALIZED_SECTION_ID_METADATA_STRING		= 0xC6BD6CDCEEF79533ull;
diff --git a/Engine/Source/Programs/Unsync/Private/UnsyncProxy.cpp b/Engine/Source/Programs/Unsync/Private/UnsyncProxy.cpp
--- a/Engine/Source/Programs/Unsync/Private/UnsyncProxy.cpp
+++ b/Engine/Source/Programs/Unsync/Private/UnsyncProxy.cpp
@@ -19,7 +19,25 @@
 
 namespace unsync {
 
-struct FUnsyncProtocolImpl : FRemoteProtocolBase
+struct FUnsyncBaseProtocolImpl : FRemoteProtocolBase
+{
+	FUnsyncBaseProtocolImpl(const FRemoteDesc&			   InRemoteDesc,
+							const FBlockRequestMap*		   InRequestMap,
+							const FRemoteProtocolFeatures& InFeatures)
+	: FRemoteProtocolBase(InRemoteDesc, InRequestMap)
+	, Features(InFeatures)
+	{
+	}
+
+	virtual TResult<FBuffer> DownloadManifest(std::string_view ManifestName) override
+	{
+		return AppError(L"Manifests can't be downloaded from UNSYNC proxy.");
+	};
+
+	const FRemoteProtocolFeatures Features;
+};
+
+struct FUnsyncProtocolImpl : FUnsyncBaseProtocolImpl
 {
 	FUnsyncProtocolImpl(const FRemoteDesc&			   InRemoteDesc,
 						const FRemoteProtocolFeatures& InFeatures,
@@ -28,10 +46,7 @@ struct FUnsyncProtocolImpl : FRemoteProtocolBase
 	virtual ~FUnsyncProtocolImpl() override;
 	virtual bool			 IsValid() const override;
 	virtual FDownloadResult	 Download(const TArrayView<FNeedBlock> NeedBlocks, const FBlockDownloadCallback& CompletionCallback) override;
-	virtual TResult<FBuffer> DownloadManifest(std::string_view ManifestName) override
-	{
-		return AppError(L"Manifests can't be downloaded from UNSYNC proxy.");
-	};
+
 	virtual void Invalidate() override;
 	virtual bool Contains(const FDirectoryManifest& Manifest) override { return true; }	 // TODO: check files on the unsync proxy
 
@@ -40,12 +55,36 @@ struct FUnsyncProtocolImpl : FRemoteProtocolBase
 	bool						 bIsConnetedToHost = false;
 	std::unique_ptr<FSocketBase> SocketHandle;
 
-	const FRemoteProtocolFeatures Features;
-
 	static void SendTelemetryEvent(const FRemoteDesc& RemoteDesc, const FTelemetryEventSyncComplete& Event);
 };
 
-FProxy::FProxy(const FRemoteDesc& RemoteDesc, const FRemoteProtocolFeatures& InFeatures, const FAuthDesc* InAuthDesc, const FBlockRequestMap* InRequestMap)
+struct FUnsyncHttpProtocolImpl : FUnsyncBaseProtocolImpl
+{
+	FUnsyncHttpProtocolImpl(const FRemoteDesc& InRemoteDesc,
+							const FRemoteProtocolFeatures& InFeatures,
+							const FBlockRequestMap* InRequestMap,
+							FProxyPool&				   InProxyPool)
+	: FUnsyncBaseProtocolImpl(InRemoteDesc, InRequestMap, InFeatures)
+	, ProxyPool(InProxyPool)
+	{
+	}
+
+	virtual bool			IsValid() const override { return bValid; }
+	virtual void			Invalidate() override { bValid = false; }
+	virtual bool			Contains(const FDirectoryManifest& Manifest) override { return true; }	// TODO: check files on the unsync proxy
+
+	virtual FDownloadResult Download(const TArrayView<FNeedBlock> NeedBlocks, const FBlockDownloadCallback& CompletionCallback) override;
+
+	const FRemoteProtocolFeatures Features;
+	bool						  bValid = true;
+	FProxyPool&					  ProxyPool;
+};
+
+FProxy::FProxy(FProxyPool&					  ProxyPool,
+			   const FRemoteDesc&			  RemoteDesc,
+			   const FRemoteProtocolFeatures& InFeatures,
+			   const FAuthDesc*				  InAuthDesc,
+			   const FBlockRequestMap*		  InRequestMap)
 {
 	UNSYNC_ASSERT(InRequestMap);
 
@@ -56,8 +95,16 @@ FProxy::FProxy(const FRemoteDesc& RemoteDesc, const FRemoteProtocolFeatures& InF
 	}
 	else if (RemoteDesc.Protocol == EProtocolFlavor::Unsync)
 	{
-		auto* Inner	 = new FUnsyncProtocolImpl(RemoteDesc, InFeatures, InAuthDesc, InRequestMap);
-		ProtocolImpl = std::unique_ptr<FRemoteProtocolBase>(Inner);
+		if (GExperimental && ProxyPool.SupportsHttp() && InFeatures.bBlockDownload)
+		{
+			auto* Inner	 = new FUnsyncHttpProtocolImpl(RemoteDesc, InFeatures, InRequestMap, ProxyPool);
+			ProtocolImpl = std::unique_ptr<FRemoteProtocolBase>(Inner);
+		}
+		else
+		{
+			auto* Inner	 = new FUnsyncProtocolImpl(RemoteDesc, InFeatures, InAuthDesc, InRequestMap);
+			ProtocolImpl = std::unique_ptr<FRemoteProtocolBase>(Inner);
+		}
 	}
 	else
 	{
@@ -69,8 +116,7 @@ FUnsyncProtocolImpl::FUnsyncProtocolImpl(const FRemoteDesc&				RemoteDesc,
 										 const FRemoteProtocolFeatures& InFeatures,
 										 const FAuthDesc*				InAuthDesc,
 										 const FBlockRequestMap*		InRequestMap)
-: FRemoteProtocolBase(RemoteDesc, InRequestMap)
-, Features(InFeatures)
+: FUnsyncBaseProtocolImpl(RemoteDesc, InRequestMap, InFeatures)
 {
 	if (RemoteDesc.TlsRequirement != ETlsRequirement::None)
 	{
@@ -222,6 +268,22 @@ FUnsyncProtocolImpl::IsValid() const
 	return bIsConnetedToHost && SocketHandle && SocketValid(*SocketHandle);
 }
 
+static void
+SortBlockRequests(std::vector<FBlockRequest>& Requests)
+{
+	std::sort(Requests.begin(),
+			  Requests.end(),
+			  [](const FBlockRequest& A, const FBlockRequest& B) -> bool
+			  {
+				  int32 FileCmp = std::memcmp(A.FilenameMd5.Data, B.FilenameMd5.Data, A.FilenameMd5.Size());
+				  if (FileCmp != 0)
+				  {
+					  return FileCmp < 0;
+				  }
+				  return A.Offset < B.Offset;
+			  });
+}
+
 FDownloadResult
 FUnsyncProtocolImpl::Download(const TArrayView<FNeedBlock> NeedBlocks, const FBlockDownloadCallback& CompletionCallback)
 {
@@ -256,15 +318,7 @@ FUnsyncProtocolImpl::Download(const TArrayView<FNeedBlock> NeedBlocks, const FBl
 		}
 	}
 
-	std::sort(Requests.begin(), Requests.end(), [](const FBlockRequest& A, const FBlockRequest& B) -> bool
-	{
-		int32 FileCmp = std::memcmp(A.FilenameMd5.Data, B.FilenameMd5.Data, A.FilenameMd5.Size());
-		if (FileCmp != 0)
-		{
-			return FileCmp;
-		}
-		return A.Offset < B.Offset;
-	});
+	SortBlockRequests(Requests);
 
 	bool bOk = bIsConnetedToHost;
 
@@ -408,6 +462,126 @@ FUnsyncProtocolImpl::Download(const TArrayView<FNeedBlock> NeedBlocks, const FBl
 	return ResultOk<FDownloadError>();
 }
 
+static std::string
+FormatBlockRequestJson(const FBlockRequestMap& RequestMap, const TArrayView<FNeedBlock> NeedBlocks)
+{
+	std::vector<FBlockRequest> Requests;
+
+	for (const FNeedBlock& Block : NeedBlocks)
+	{
+		if (const FBlockRequest* Request = RequestMap.FindRequest(Block.Hash))
+		{
+			Requests.push_back(*Request);
+		}
+	}
+
+	SortBlockRequests(Requests);
+
+	std::string Output;
+	Output += "{ ";  // main object
+	Output += "\"files\": [\n";
+
+	static const FHash128 InvalidHash = {};
+	FHash128	   FilenameHash = InvalidHash;
+
+	uint32 BlockIndex = 0;
+	for (const FBlockRequest& Request : Requests)
+	{
+		if (FilenameHash != Request.FilenameMd5)
+		{
+			if (FilenameHash != InvalidHash)
+			{
+				Output += "]},\n";  // close blocks arary and file object
+			}
+
+			const std::string* FilenameUtf8 = RequestMap.FindFile(Request.FilenameMd5);
+			UNSYNC_ASSERTF(FilenameUtf8, L"Could not find file in the block request map");
+
+			// Start file object and blocks array
+			Output += "{";
+
+			std::string EscapedFilenameUtf8 = StringEscape(*FilenameUtf8);
+			FormatJsonKeyValueStr(Output, "name", EscapedFilenameUtf8, ", ");
+			Output += "\"blocks\": [\n";
+
+			BlockIndex = 0;
+			FilenameHash = Request.FilenameMd5;
+		}
+
+		if (BlockIndex != 0)
+		{
+			Output += ",\n";
+		}
+
+		FGenericBlock Block;
+		Block.HashStrong = FGenericHash::FromBlake3_128(Request.BlockHash);
+		Block.Offset	 = Request.Offset;
+		Block.Size		 = CheckedNarrow(Request.Size);
+
+		FormatJsonBlock(Output, Block);
+
+		++BlockIndex;
+	}
+
+	if (FilenameHash != InvalidHash)
+	{
+		Output += "]}\n";	// close blocks arary and file object
+	}
+
+	Output += "]\n";  // files array
+	Output += "}\n";  // main object
+
+	return Output;
+}
+
+FDownloadResult
+FUnsyncHttpProtocolImpl::Download(const TArrayView<FNeedBlock> NeedBlocks, const FBlockDownloadCallback& CompletionCallback)
+{
+	if (!IsValid())
+	{
+		return FDownloadResult(EDownloadRetryMode::Abort);
+	}
+
+	std::string RequestJson = FormatBlockRequestJson(*RequestMap, NeedBlocks);
+
+	auto ChunkCallback = [&CompletionCallback](FHttpResponse& Response)
+	{
+		if (Response.Success() && Response.Buffer.Size())
+		{
+			FBuffer			 DecompressedBuffer = Decompress(Response.Buffer);
+			FDownloadedBlock DownloadedBlock;
+			DownloadedBlock.bCompressed		 = false;
+			DownloadedBlock.Data			 = DecompressedBuffer.Data();
+			DownloadedBlock.DecompressedSize = DecompressedBuffer.Size();
+			FHash128 DecompressedHash		 = HashBlake3Bytes<FHash128>(DecompressedBuffer.Data(), DecompressedBuffer.Size());
+			CompletionCallback(DownloadedBlock, DecompressedHash);
+		}
+
+		Response.Buffer.Clear();
+	};
+
+	FPooledHttpConnection HttpConnection(ProxyPool);
+
+	FHttpRequest Request;
+	Request.Method			   = EHttpMethod::POST;
+	Request.PayloadContentType = EHttpContentType::Application_Json;
+	Request.Payload.Data	   = reinterpret_cast<const uint8*>(RequestJson.data());
+	Request.Payload.Size	   = RequestJson.length();
+	Request.Url				   = "/api/v1/blocks";
+
+	FHttpResponse Response = HttpRequest(HttpConnection, Request, ChunkCallback);
+
+	if (Response.Success())
+	{
+		return ResultOk<FDownloadError>();
+	}
+	else
+	{
+		UNSYNC_ERROR(L"Failed to complete block request. HTTP error code: %d.", Response.Code);
+		return FDownloadError(EDownloadRetryMode::Abort);
+	}
+}
+
 void
 FUnsyncProtocolImpl::Invalidate()
 {
@@ -538,6 +712,10 @@ ProxyQuery::Hello(FHttpConnection& HttpConnection, const FAuthDesc* OptAuthDesc)
 				{
 					Result.Features.bFileDownload = true;
 				}
+				else if (Elem.string_value() == "blocks")
+				{
+					Result.Features.bBlockDownload = true;
+				}
 			}
 		}
 	}
@@ -1017,7 +1195,11 @@ FProxyPool::FProxyPool(const FRemoteDesc& InRemoteDesc, const FAuthDesc* InAuthD
 					   RemoteDesc.Host.Address.c_str(),
 					   RemoteDesc.Host.Port);
 
-		TResult<ProxyQuery::FHelloResponse> Response = ProxyQuery::Hello(RemoteDesc, AuthDesc);
+		std::unique_ptr<FHttpConnection> HttpConnection = HttpPool->Acquire();
+
+		TResult<ProxyQuery::FHelloResponse> Response = ProxyQuery::Hello(*HttpConnection, AuthDesc);
+
+		HttpPool->Release(std::move(HttpConnection));
 
 		if (Response.IsError())
 		{
@@ -1064,7 +1246,7 @@ FProxyPool::Alloc()
 
 	if (!Result || !Result->IsValid())
 	{
- 		Result = std::make_unique<FProxy>(RemoteDesc, Features, AuthDesc, &RequestMap);
+ 		Result = std::make_unique<FProxy>(*this, RemoteDesc, Features, AuthDesc, &RequestMap);
 	}
 
 	return Result;
diff --git a/Engine/Source/Programs/Unsync/Private/UnsyncProxy.h b/Engine/Source/Programs/Unsync/Private/UnsyncProxy.h
--- a/Engine/Source/Programs/Unsync/Private/UnsyncProxy.h
+++ b/Engine/Source/Programs/Unsync/Private/UnsyncProxy.h
@@ -99,6 +99,7 @@ struct FRemoteProtocolFeatures
 	bool bDirectoryListing = false;
 	bool bFileDownload	   = false;
 	bool bDownloadByHash   = false;
+	bool bBlockDownload	   = false;
 };
 
 struct FTelemetryEventSyncComplete
@@ -144,7 +145,8 @@ struct FRemoteProtocolBase
 class FProxy
 {
 public:
-	FProxy(const FRemoteDesc&			  InRemoteDesc,
+	FProxy(FProxyPool&					  ProxyPool,
+		   const FRemoteDesc&			  InRemoteDesc,
 		   const FRemoteProtocolFeatures& InFeatures,
 		   const FAuthDesc*				  InAuthDesc,
 		   const FBlockRequestMap*		  InRequestMap);
@@ -172,6 +174,8 @@ public:
 	std::unique_ptr<FHttpConnection> AllocHttp();
 	void							 DeallocHttp(std::unique_ptr<FHttpConnection>&& Connection);
 
+	bool SupportsHttp() const { return HttpPool.has_value(); }
+
 	void Invalidate();
 	bool IsValid() const;
 
diff --git a/Engine/Source/Programs/Unsync/Private/UnsyncUtil.cpp b/Engine/Source/Programs/Unsync/Private/UnsyncUtil.cpp
--- a/Engine/Source/Programs/Unsync/Private/UnsyncUtil.cpp
+++ b/Engine/Source/Programs/Unsync/Private/UnsyncUtil.cpp
@@ -87,18 +87,37 @@ FormatJsonKeyValueStr(std::wstring& Output, std::wstring_view K, std::wstring_vi
 	fmt::format_to(std::back_inserter(Output), L"\"{}\": \"{}\"{}", K, V, Suffix);
 }
 
+void
+FormatJsonKeyValueStr(std::string& Output, std::string_view K, std::string_view V, std::string_view Suffix)
+{
+	fmt::format_to(std::back_inserter(Output), "\"{}\": \"{}\"{}", K, V, Suffix);
+}
+
 void
 FormatJsonKeyValueUInt(std::wstring& Output, std::wstring_view K, uint64 V, std::wstring_view Suffix)
 {
 	fmt::format_to(std::back_inserter(Output), L"\"{}\": {}{}", K, V, Suffix);
 }
 
+void
+FormatJsonKeyValueUInt(std::string& Output, std::string_view K, uint64 V, std::string_view Suffix)
+{
+	fmt::format_to(std::back_inserter(Output), "\"{}\": {}{}", K, V, Suffix);
+}
+
+
 void
 FormatJsonKeyValueBool(std::wstring& Output, std::wstring_view K, bool V, std::wstring_view Suffix)
 {
 	fmt::format_to(std::back_inserter(Output), L"\"{}\": {}{}", K, V ? L"true" : L"false", Suffix);
 }
 
+void
+FormatJsonKeyValueBool(std::string& Output, std::string_view K, bool V, std::string_view Suffix)
+{
+	fmt::format_to(std::back_inserter(Output), "\"{}\": {}{}", K, V ? "true" : "false", Suffix);
+}
+
 void
 FormatJsonBlock(std::wstring& Output, const FGenericBlock& Block)
 {
@@ -121,6 +140,28 @@ FormatJsonBlock(std::wstring& Output, const FGenericBlock& Block)
 	Output += L"}";
 }
 
+void
+FormatJsonBlock(std::string& Output, const FGenericBlock& Block)
+{
+	Output += "{";
+
+	static const size_t MaxHashLen = 2 * sizeof(Block.HashStrong.Data);
+	char				HashChars[MaxHashLen];
+
+	uint64			 HashLen = BytesToHexChars(HashChars, MaxHashLen, Block.HashStrong.Data, Block.HashStrong.Size());
+	std::string_view HashStr = std::string_view(HashChars, HashLen);
+
+	FormatJsonKeyValueUInt(Output, "offset", Block.Offset, ", ");
+	FormatJsonKeyValueUInt(Output, "size", Block.Size, ", ");
+	if (Block.HashWeak != 0)
+	{
+		FormatJsonKeyValueUInt(Output, "hash_weak", Block.HashWeak, ", ");
+	}
+	FormatJsonKeyValueStr(Output, "hash_strong", HashStr);
+
+	Output += "}";
+}
+
 void
 FormatJsonBlockArray(std::wstring& Output, const FGenericBlockArray& Blocks)
 {
@@ -140,6 +181,25 @@ FormatJsonBlockArray(std::wstring& Output, const FGenericBlockArray& Blocks)
 	Output += L"]";
 }
 
+void
+FormatJsonBlockArray(std::string& Output, const FGenericBlockArray& Blocks)
+{
+	Output += "[\n";
+	uint64 BlockIndex = 0;
+	for (const FGenericBlock& Block : Blocks)
+	{
+		if (BlockIndex != 0)
+		{
+			Output += ",\n";
+		}
+
+		FormatJsonBlock(Output, Block);
+
+		++BlockIndex;
+	}
+	Output += "]";
+}
+
 
 FTimingLogger::FTimingLogger(const char* InName, ELogLevel InLogLevel, bool bInEnabled)
 : bEnabled(bInEnabled)
diff --git a/Engine/Source/Programs/Unsync/Private/UnsyncUtil.h b/Engine/Source/Programs/Unsync/Private/UnsyncUtil.h
--- a/Engine/Source/Programs/Unsync/Private/UnsyncUtil.h
+++ b/Engine/Source/Programs/Unsync/Private/UnsyncUtil.h
@@ -374,7 +374,10 @@ bool LooksLikeHash160(const std::wstring_view Str);
 
 // Json formatting helpers
 void FormatJsonKeyValueStr(std::wstring& Output, std::wstring_view K, std::wstring_view V, std::wstring_view Suffix = {});
+void FormatJsonKeyValueStr(std::string& Output, std::string_view K, std::string_view V, std::string_view Suffix = {});
 void FormatJsonKeyValueUInt(std::wstring& Output, std::wstring_view K, uint64 V, std::wstring_view Suffix = {});
+void FormatJsonKeyValueUInt(std::string& Output, std::string_view K, uint64 V, std::string_view Suffix = {});
 void FormatJsonKeyValueBool(std::wstring& Output, std::wstring_view K, bool V, std::wstring_view Suffix = {});
+void FormatJsonKeyValueBool(std::string& Output, std::string_view K, bool V, std::string_view Suffix = {});
 
 }  // namespace unsync
