unsync - Undo accidental debug code

diff --git a/Engine/Source/Programs/Unsync/Private/UnsyncMain.cpp b/Engine/Source/Programs/Unsync/Private/UnsyncMain.cpp
--- a/Engine/Source/Programs/Unsync/Private/UnsyncMain.cpp
+++ b/Engine/Source/Programs/Unsync/Private/UnsyncMain.cpp
@@ -880,8 +880,6 @@ InnerMain(int Argc, char** Argv)
 		}
 	}
 
-	RemoteDesc.TlsRequirement = ETlsRequirement::None; // nocheckin
-
 	if (Cli.got_subcommand(SubHash))
 	{
 		FCmdHashOptions HashOptions;
