unsync - Send strong hash algorithm id with the block requests

diff --git a/Engine/Source/Programs/Unsync/Private/UnsyncProxy.cpp b/Engine/Source/Programs/Unsync/Private/UnsyncProxy.cpp
--- a/Engine/Source/Programs/Unsync/Private/UnsyncProxy.cpp
+++ b/Engine/Source/Programs/Unsync/Private/UnsyncProxy.cpp
@@ -465,6 +465,8 @@ FUnsyncProtocolImpl::Download(const TArrayView<FNeedBlock> NeedBlocks, const FBl
 static std::string
 FormatBlockRequestJson(const FBlockRequestMap& RequestMap, const TArrayView<FNeedBlock> NeedBlocks)
 {
+	const char* StrongHashAlgorithm = ToString(RequestMap.GetStrongHasher());
+
 	std::vector<FBlockRequest> Requests;
 
 	for (const FNeedBlock& Block : NeedBlocks)
@@ -478,7 +480,11 @@ FormatBlockRequestJson(const FBlockRequestMap& RequestMap, const TArrayView<FNee
 	SortBlockRequests(Requests);
 
 	std::string Output;
+
 	Output += "{ ";  // main object
+
+	FormatJsonKeyValueStr(Output, "hash_strong", StrongHashAlgorithm, ",\n");
+
 	Output += "\"files\": [\n";
 
 	static const FHash128 InvalidHash = {};
@@ -544,6 +550,9 @@ FUnsyncHttpProtocolImpl::Download(const TArrayView<FNeedBlock> NeedBlocks, const
 
 	std::string RequestJson = FormatBlockRequestJson(*RequestMap, NeedBlocks);
 
+	// TODO: send compressed requests for ~3x upload bandwidth reduction
+	// FBuffer CompressedRequest = Compress(reinterpret_cast<const uint8*>(RequestJson.data()), RequestJson.size(), 9);
+
 	auto ChunkCallback = [&CompletionCallback](FHttpResponse& Response)
 	{
 		if (Response.Success() && Response.Buffer.Size())
