unsync - Use deferred open mechanism for base files

* Avoid opening base files unnecessarily

diff --git a/Engine/Source/Programs/Unsync/Private/UnsyncCore.cpp b/Engine/Source/Programs/Unsync/Private/UnsyncCore.cpp
--- a/Engine/Source/Programs/Unsync/Private/UnsyncCore.cpp
+++ b/Engine/Source/Programs/Unsync/Private/UnsyncCore.cpp
@@ -1854,11 +1854,20 @@ SyncDirectory(const FSyncDirectoryOptions& SyncOptions)
 								 &ProxyPool](const FFileSyncTask& Item, FBlockCache* BlockCache, bool bBackground) {
 			UNSYNC_VERBOSE(L"Copy '%ls' (%ls)", Item.TargetFilePath.wstring().c_str(), (Item.NeedBytesFromBase) ? L"partial" : L"full");
 
-			std::unique_ptr<FNativeFile> BaseFile;
-			if (Item.IsBaseValid() && !Item.NeedList.Base.empty())
-			{
-				BaseFile = std::make_unique<FNativeFile>(Item.BaseFilePath, EFileMode::ReadOnlyUnbuffered);
-			}
+			FDeferredOpenReader BaseFile(
+				[&Item]
+				{
+					if (Item.IsBaseValid())
+					{
+						UNSYNC_VERBOSE(L"Opening base file '%ls'", Item.BaseFilePath.wstring().c_str());
+						LogStatus(Item.BaseFilePath.wstring().c_str(), L"Opening base file");
+						return std::unique_ptr<FIOReader>(new FNativeFile(Item.BaseFilePath, EFileMode::ReadOnlyUnbuffered));
+					}
+					else
+					{
+						return std::unique_ptr<FIOReader>(new FNullReaderWriter(0));
+					}
+				});
 
 			const FGenericBlockArray& SourceBlocks	  = Item.SourceManifest->Blocks;
 			uint32					  SourceBlockSize = Item.SourceManifest->BlockSize;
@@ -1872,7 +1881,7 @@ SyncDirectory(const FSyncDirectoryOptions& SyncOptions)
 			SyncFileOptions.bValidateTargetFiles = SyncOptions.bValidateTargetFiles;
 
 			FFileSyncResult SyncResult =
-				SyncFile(Item.NeedList, Item.ResolvedSourceFilePath, SourceBlocks, *BaseFile.get(), Item.TargetFilePath, SyncFileOptions);
+				SyncFile(Item.NeedList, Item.ResolvedSourceFilePath, SourceBlocks, BaseFile, Item.TargetFilePath, SyncFileOptions);
 
 			LogStatus(Item.TargetFilePath.wstring().c_str(), SyncResult.Succeeded() ? L"Succeeded" : L"Failed");
 
@@ -1884,7 +1893,7 @@ SyncDirectory(const FSyncDirectoryOptions& SyncOptions)
 
 				if (!GDryRun)
 				{
-					BaseFile = nullptr;
+					BaseFile.Close();
 					SetFileMtime(Item.TargetFilePath, Item.SourceManifest->Mtime);
 					if (Item.SourceManifest->bReadOnly)
 					{
diff --git a/Engine/Source/Programs/Unsync/Private/UnsyncFile.h b/Engine/Source/Programs/Unsync/Private/UnsyncFile.h
--- a/Engine/Source/Programs/Unsync/Private/UnsyncFile.h
+++ b/Engine/Source/Programs/Unsync/Private/UnsyncFile.h
@@ -380,7 +380,7 @@ struct FDeferredOpenReader : FIOReader
 	virtual void   FlushOne() override { GetOrOpenInner()->FlushOne(); }
 	virtual uint64 GetSize() override { return GetOrOpenInner()->GetSize(); }
 	virtual bool   IsValid() override { return GetOrOpenInner()->IsValid(); }
-	virtual void   Close() override { GetOrOpenInner()->Close(); }
+	virtual void   Close() override { if (Inner) { Inner->Close(); } }
 	virtual int32  GetError() override { return GetOrOpenInner()->GetError(); }
 
 	// IORead
