unsync - Add unit tests for chunking algorithm, add utility function to compute rolling window hash threshold to achieve specific average block size

diff --git a/Engine/Source/Programs/Unsync/Private/UnsyncChunking.cpp b/Engine/Source/Programs/Unsync/Private/UnsyncChunking.cpp
--- a/Engine/Source/Programs/Unsync/Private/UnsyncChunking.cpp
+++ b/Engine/Source/Programs/Unsync/Private/UnsyncChunking.cpp
@@ -25,6 +25,14 @@ ComputeMaxVariableBlockSize(uint32 BlockSize)
 	return std::min(MAX_BLOCK_SIZE, std::max(BlockSize, 4096u) * 4);  // changing this invalidates cached blocks
 }
 
+inline uint32
+ComputeWindowHashThreshold(uint32 TargetSize)
+{
+	const uint32 MinSize = ComputeMinVariableBlockSize(TargetSize);
+	UNSYNC_ASSERT(TargetSize > MinSize);
+	return uint32((1ull << 32) / (TargetSize-MinSize));
+}
+
 template<typename WeakHasher>
 FComputeBlocksResult
 ComputeBlocksVariableT(FIOReader& Reader, const FComputeBlocksParams& Params)
@@ -82,6 +90,8 @@ ComputeBlocksVariableT(FIOReader& Reader, const FComputeBlocksParams& Params)
 			UNSYNC_FATAL(L"Expected to read %lld bytes from input, but %lld was actually read.", ThisTaskSize, ReadBytesForTask);
 		}
 
+		const uint32 ChunkWindowHashThreshold = ComputeWindowHashThreshold(Params.BlockSize);
+
 		auto ScanTask = [&Tasks,
 						 &IoSemaphore,
 						 &BufferPool,
@@ -92,6 +102,7 @@ ComputeBlocksVariableT(FIOReader& Reader, const FComputeBlocksParams& Params)
 						 ScanTaskBuffer,
 						 TaskIndex,
 						 ThisTaskSize,
+						 ChunkWindowHashThreshold,
 						 TargetMacroBlockSize,
 						 MinimumMacroBlockSize,
 						 MaximumMacroBlockSize,
@@ -114,6 +125,7 @@ ComputeBlocksVariableT(FIOReader& Reader, const FComputeBlocksParams& Params)
 						   MaximumMacroBlockSize,
 						   MaximumBlockSize,
 						   MinimumMacroBlockSize,
+						   ChunkWindowHashThreshold,
 						   DataEnd,
 						   &LastBlockEnd,
 						   &Task,
@@ -124,10 +136,6 @@ ComputeBlocksVariableT(FIOReader& Reader, const FComputeBlocksParams& Params)
 						   &CurrentMacroBlock,
 						   &SourceInfo](const uint8* WindowBegin, const uint8* WindowEnd, uint32 WindowHash) UNSYNC_ATTRIB_FORCEINLINE
 			{
-				// WARNING: Changing this invalidates some of the previously cached blocks.
-				// TODO: compute based on target average block size
-				const uint32 ChunkWindowHashThreshold = 0x20000;
-
 				const bool	 bLastBlock	   = WindowEnd == DataEnd;
 				const uint64 ThisBlockSize = uint64(WindowEnd - LastBlockEnd);
 
@@ -512,4 +520,119 @@ ComputeBlocks(const uint8* Data, uint64 Size, const FComputeBlocksParams& Params
 	return ComputeBlocks(DataReader, Params);
 }
 
+static FBuffer GenerateTestData(uint64 Size, uint32 Seed = 1234)
+{
+	FBuffer Buffer(Size, 0);
+	uint32	Rng = Seed;
+	for (uint64 I = 0; I < Buffer.Size(); ++I)
+	{
+		Buffer[I] = Xorshift32(Rng) & 0xFF;
+	}
+	return Buffer;
+}
+
+void
+TestChunking()
+{
+	UNSYNC_LOG(L"TestChunking()");
+	UNSYNC_LOG_INDENT;
+
+	UNSYNC_LOG(L"Generating data");
+
+	{
+		const uint32 Threshold = ComputeWindowHashThreshold(uint32(64_KB));
+		const uint32 ExpectedValue = 0x20000;
+		if (Threshold != ExpectedValue)
+		{
+			UNSYNC_ERROR("Expected window hash threshold for 64 KB target block size: 0x%08x, actual value: 0x%08x", ExpectedValue, Threshold);
+		}
+	}
+
+	FBuffer Buffer = GenerateTestData(1_GB);
+
+	UNSYNC_LOG(L"Testing expected chunk boundaries");
+
+	{
+		UNSYNC_LOG_INDENT;
+
+		FComputeBlocksParams Params;
+		Params.bNeedMacroBlocks				   = false;
+		Params.BlockSize					   = uint32(64_KB);
+		Params.Algorithm.WeakHashAlgorithmId   = EWeakHashAlgorithmID::BuzHash;
+		Params.Algorithm.StrongHashAlgorithmId = EStrongHashAlgorithmID::Blake3_160;
+
+		FMemReader			 Reader(Buffer.Data(), 1_MB);
+		FComputeBlocksResult Blocks = ComputeBlocksVariable(Reader, Params);
+
+		uint64 NumBlocks = Blocks.Blocks.size();
+		uint64 AvgSize	 = Buffer.Size() / NumBlocks;
+
+		static constexpr uint32 NumExpectedBlocks = 18;
+
+		// clang-format off
+		const uint64 ExpectedOffsets[NumExpectedBlocks] = {
+			0, 34577, 128471, 195115, 238047, 297334, 358754, 396031,
+			462359, 508658, 601550, 702021, 754650, 790285, 854987, 887998,
+			956848, 1042406
+		};
+		// clang-format on
+
+		UNSYNC_LOG(L"Generated blocks: %llu, average size: %llu KB", llu(NumBlocks), llu(AvgSize) / 1024);
+
+		if (Blocks.Blocks.size() != NumExpectedBlocks)
+		{
+			UNSYNC_ERROR("Expected blocks: %llu, actual number: %llu", llu(NumExpectedBlocks), llu(Blocks.Blocks.size()));
+		}
+
+		for (uint32 ChunkIndex = 0; ChunkIndex < std::min<size_t>(NumExpectedBlocks, Blocks.Blocks.size()); ++ChunkIndex)
+		{
+			const FGenericBlock& Block = Blocks.Blocks[ChunkIndex];
+			UNSYNC_LOG(L" - [%2d] offset: %llu, size: %llu, weak_hash: 0x%08x",
+					   ChunkIndex,
+					   llu(Block.Offset),
+					   llu(Block.Size),
+					   Block.HashWeak);
+
+			if (ExpectedOffsets[ChunkIndex] != Block.Offset)
+			{
+				UNSYNC_ERROR("Expected block at offset: %llu, actual offset: %llu", ExpectedOffsets[ChunkIndex], Block.Offset);
+			}
+		}
+	}
+
+	const uint32 TestChunkSizesKB[] = {8, 16, 32, 64, 96, 128, 160, 192, 256};
+
+	UNSYNC_LOG(L"Testing average chunk size");
+
+	for (uint32 ChunkSizeKB : TestChunkSizesKB)
+	{
+		UNSYNC_LOG_INDENT;
+
+		FComputeBlocksParams Params;
+		Params.bNeedMacroBlocks				   = false;
+		Params.BlockSize					   = uint32(ChunkSizeKB * 1024);
+		Params.Algorithm.WeakHashAlgorithmId   = EWeakHashAlgorithmID::BuzHash;
+		Params.Algorithm.StrongHashAlgorithmId = EStrongHashAlgorithmID::Blake3_160;
+
+		const uint32 Threshold = ComputeWindowHashThreshold(Params.BlockSize);
+		UNSYNC_LOG(L"ComputeBlocksVariableT<FBuzHash>, %d KB target, window hash threshold: 0x%08x", ChunkSizeKB, Threshold);
+
+		FMemReader			 Reader(Buffer);
+		FComputeBlocksResult Blocks = ComputeBlocksVariable(Reader, Params);
+
+		uint64 NumBlocks = Blocks.Blocks.size();
+		uint64 AvgSize	 = Buffer.Size() / NumBlocks;
+		
+		int64 AbsDiff = std::abs(int64(Params.BlockSize) - int64(AvgSize));
+		double AbsDiffPct = 100.0 * double(AbsDiff) / double(Params.BlockSize);
+
+		UNSYNC_LOG(L"Generated blocks: %llu, average size: %llu KB, error %.2f %%", llu(NumBlocks), llu(AvgSize) / 1024, AbsDiffPct);
+
+		if (AbsDiffPct > 5.0)
+		{
+			UNSYNC_ERROR(L"Average block size is significantly different from target");
+		}
+	}
+}
+
 }  // namespace unsync
diff --git a/Engine/Source/Programs/Unsync/Private/UnsyncTest.cpp b/Engine/Source/Programs/Unsync/Private/UnsyncTest.cpp
--- a/Engine/Source/Programs/Unsync/Private/UnsyncTest.cpp
+++ b/Engine/Source/Programs/Unsync/Private/UnsyncTest.cpp
@@ -884,6 +884,12 @@ RunTests(const std::string& Preset)
 		TestBasicHash();
 	}
 
+	if (Preset == "chunking" || Preset == "all")
+	{
+		extern void TestChunking();
+		TestChunking();
+	}
+
 	if (Preset == "sync" || Preset == "all")
 	{
 		for (auto Weak : WeakList)
