unsync - Reduce chunking test buffer to from 1GB to 128MB, add validation for number of generated chunks and report median chunk size

diff --git a/Engine/Source/Programs/Unsync/Private/UnsyncChunking.cpp b/Engine/Source/Programs/Unsync/Private/UnsyncChunking.cpp
--- a/Engine/Source/Programs/Unsync/Private/UnsyncChunking.cpp
+++ b/Engine/Source/Programs/Unsync/Private/UnsyncChunking.cpp
@@ -548,7 +548,7 @@ TestChunking()
 		}
 	}
 
-	FBuffer Buffer = GenerateTestData(1_GB);
+	FBuffer Buffer = GenerateTestData(128_MB);
 
 	UNSYNC_LOG(L"Testing expected chunk boundaries");
 
@@ -600,14 +600,19 @@ TestChunking()
 		}
 	}
 
-	const uint32 TestChunkSizesKB[] = {8, 16, 32, 64, 96, 128, 160, 192, 256};
+	static constexpr uint32 NumConfigs		 = 9;
+	const uint32			TestChunkSizesKB[NumConfigs] = {8, 16, 32, 64, 96, 128, 160, 192, 256};
+	const uint32			ExpectedNumChunks[NumConfigs] = {16442, 8146, 4089, 2019, 1362, 1012, 811, 681, 503};
 
 	UNSYNC_LOG(L"Testing average chunk size");
 
-	for (uint32 ChunkSizeKB : TestChunkSizesKB)
+	for (uint32 ConfigIndex = 0; ConfigIndex < NumConfigs; ++ConfigIndex)
 	{
 		UNSYNC_LOG_INDENT;
 
+		const uint32 ChunkSizeKB = TestChunkSizesKB[ConfigIndex];
+		const uint32 ExpectedCount = ExpectedNumChunks[ConfigIndex];
+
 		FComputeBlocksParams Params;
 		Params.bNeedMacroBlocks				   = false;
 		Params.BlockSize					   = uint32(ChunkSizeKB * 1024);
@@ -626,12 +631,28 @@ TestChunking()
 		int64 AbsDiff = std::abs(int64(Params.BlockSize) - int64(AvgSize));
 		double AbsDiffPct = 100.0 * double(AbsDiff) / double(Params.BlockSize);
 
-		UNSYNC_LOG(L"Generated blocks: %llu, average size: %llu KB, error %.2f %%", llu(NumBlocks), llu(AvgSize) / 1024, AbsDiffPct);
+		// Compute median block size
+		std::sort(Blocks.Blocks.begin(),
+				  Blocks.Blocks.end(),
+				  [](const FGenericBlock& A, const FGenericBlock& B) { return A.Size < B.Size; });
+
+		const uint64 MedianSize = Blocks.Blocks[Blocks.Blocks.size() / 2].Size;
+
+		UNSYNC_LOG(L"Generated blocks: %llu, average size: %llu KB, median: %llu KB, average error %.2f %%",
+				   llu(NumBlocks),
+				   llu(AvgSize) / 1024,
+				   llu(MedianSize) / 1024,
+				   AbsDiffPct);
 
 		if (AbsDiffPct > 5.0)
 		{
 			UNSYNC_ERROR(L"Average block size is significantly different from target");
 		}
+
+		if (ExpectedCount != NumBlocks)
+		{
+			UNSYNC_ERROR(L"Expected to generate blocks: %llu, actual: %llu", llu(ExpectedCount), llu(NumBlocks));
+		}
 	}
 }
 
