unsync - Groundwork for downloading Horde artifacts

* Recognize Horde artifact URLs
* Download and decode manifest from Horde

diff --git a/Engine/Source/Programs/Unsync/CMakeLists.txt b/Engine/Source/Programs/Unsync/CMakeLists.txt
--- a/Engine/Source/Programs/Unsync/CMakeLists.txt
+++ b/Engine/Source/Programs/Unsync/CMakeLists.txt
@@ -71,6 +71,8 @@ set(src
 	Private/UnsyncHash.cpp
 	Private/UnsyncHash.h
 	Private/UnsyncHashTable.h
+	Private/UnsyncHorde.cpp
+	Private/UnsyncHorde.h
 	Private/UnsyncHttp.cpp
 	Private/UnsyncHttp.h
 	Private/UnsyncJupiter.cpp
diff --git a/Engine/Source/Programs/Unsync/Private/UnsyncCmdPack.cpp b/Engine/Source/Programs/Unsync/Private/UnsyncCmdPack.cpp
--- a/Engine/Source/Programs/Unsync/Private/UnsyncCmdPack.cpp
+++ b/Engine/Source/Programs/Unsync/Private/UnsyncCmdPack.cpp
@@ -931,7 +931,10 @@ SyncDirectoryFromPack(const FPath& OutputRoot, const FPackDatabase& PackDb, cons
 
 			TargetFile.Close();
 
-			SetFileMtime(TargetFilePath, FileManifest.Mtime);
+			if (FileManifest.Mtime)
+			{
+				SetFileMtime(TargetFilePath, FileManifest.Mtime);
+			}
 
 			if (FileManifest.bReadOnly)
 			{
diff --git a/Engine/Source/Programs/Unsync/Private/UnsyncCmdSync.cpp b/Engine/Source/Programs/Unsync/Private/UnsyncCmdSync.cpp
--- a/Engine/Source/Programs/Unsync/Private/UnsyncCmdSync.cpp
+++ b/Engine/Source/Programs/Unsync/Private/UnsyncCmdSync.cpp
@@ -63,19 +63,24 @@ CmdSync(const FCmdSyncOptions& Options)
 	if (ProxyPool.IsValid())
 	{
 		const FRemoteProtocolFeatures& Features = ProxyPool.GetFeatures();
-		if (Features.bFileDownload && Features.bDirectoryListing)
+		if (ProxyPool.RemoteDesc.Protocol == EProtocolFlavor::Horde)
+		{
+			UNSYNC_LOG(L"Horde server supports direct file, manifest and block download");
+			bSourceFileSystemRequired = false;
+		}
+		else if (Features.bFileDownload && Features.bDirectoryListing)
 		{
 			UNSYNC_LOG(L"Server supports direct file access");
 			bSourceFileSystemRequired = false;
 		}
-		else if (Features.bDownloadByHash && bSourceIsManifestHash)
+		else if (Features.bManifestDownload && bSourceIsManifestHash)
 		{
-			UNSYNC_LOG(L"Server supports access by manifest hash");
+			UNSYNC_LOG(L"Server supports access by manifest ID");
 			bSourceFileSystemRequired = false;
 		}
 		else
 		{
-			UNSYNC_VERBOSE2(L"Server does not support direct file access or download by manifest hash. Source file system access is required.");
+			UNSYNC_VERBOSE2(L"Server does not support direct file access or download by manifest ID. Source file system access is required.");
 		}
 	}
 
@@ -114,7 +119,7 @@ CmdSync(const FCmdSyncOptions& Options)
 		
 		if (bSourceIsManifestHash)
 		{
-			UNSYNC_ERROR(L"Sync overlay option is not compatible with sync by manifest hash.");
+			UNSYNC_ERROR(L"Sync overlay option is not compatible with sync by manifest ID.");
 			return 1;
 		}
 
@@ -158,9 +163,9 @@ CmdSync(const FCmdSyncOptions& Options)
 			{
 				SyncOptions.SourceType = ESourceType::FileSystem;
 			}
-			else if (bSourceIsManifestHash)
+			else if (bSourceIsManifestHash || ProxyPool.RemoteDesc.Protocol == EProtocolFlavor::Horde)
 			{
-				SyncOptions.SourceType = ESourceType::ServerWithManifestHash;
+				SyncOptions.SourceType = ESourceType::ServerWithManifestId;
 			}
 			else
 			{
diff --git a/Engine/Source/Programs/Unsync/Private/UnsyncCore.cpp b/Engine/Source/Programs/Unsync/Private/UnsyncCore.cpp
--- a/Engine/Source/Programs/Unsync/Private/UnsyncCore.cpp
+++ b/Engine/Source/Programs/Unsync/Private/UnsyncCore.cpp
@@ -935,7 +935,10 @@ CopyFileIfPossiblyDifferent(FProxyFileSystem&	   FileSystem,
 		}
 
 		const bool bAllowInDryRun = true;
-		SetFileMtime(Target, Source.Entry.Mtime, bAllowInDryRun);
+		if (Source.Entry.Mtime)
+		{
+			SetFileMtime(Target, Source.Entry.Mtime, bAllowInDryRun);
+		}
 	}
 
 	return true;
@@ -1186,7 +1189,7 @@ SyncDirectory(const FSyncDirectoryOptions& SyncOptions)
 
 	const bool bFileSystemSource = SyncOptions.SourceType == ESourceType::FileSystem;
 	const bool bServerSource =
-		SyncOptions.SourceType == ESourceType::Server || SyncOptions.SourceType == ESourceType::ServerWithManifestHash;
+		SyncOptions.SourceType == ESourceType::Server || SyncOptions.SourceType == ESourceType::ServerWithManifestId;
 
 	UNSYNC_ASSERT(bFileSystemSource || bServerSource);
 
@@ -1254,11 +1257,11 @@ SyncDirectory(const FSyncDirectoryOptions& SyncOptions)
 
 	std::vector<FPackIndexDatabase> PackIndexFiles;
 
-	if (SyncOptions.SourceType == ESourceType::ServerWithManifestHash)
+	if (SyncOptions.SourceType == ESourceType::ServerWithManifestId)
 	{
 		if (!ProxyPool.IsValid())
 		{
-			UNSYNC_ERROR(L"Remote server connection is required when syncing by manifest hash");
+			UNSYNC_ERROR(L"Remote server connection is required when syncing by manifest ID");
 			return false;
 		}
 
@@ -1266,14 +1269,12 @@ SyncDirectory(const FSyncDirectoryOptions& SyncOptions)
 
 		std::unique_ptr<FProxy> Proxy = ProxyPool.Alloc();
 
-		std::string		 SourceManifestName = ConvertWideToUtf8(SyncOptions.Source.wstring());
-		TResult<FBuffer> DownloadResult		= Proxy->DownloadManifest(SourceManifestName);
-		if (FBuffer* ManifestBuffer = DownloadResult.TryData())
+		std::string					SourceManifestName = ConvertWideToUtf8(SyncOptions.Source.wstring());
+		TResult<FDirectoryManifest> DownloadResult	   = Proxy->DownloadManifest(SourceManifestName);
+		if (FDirectoryManifest* Manifest = DownloadResult.TryData())
 		{
-			FMemReader		Reader(*ManifestBuffer);
-			FIOReaderStream Stream(Reader);
-			FPath			EmptyRoot;	// Don't have a sensible path when not using file system as source
-			bSourceManifestOk = LoadDirectoryManifest(SourceDirectoryManifest, EmptyRoot, Stream);
+			bSourceManifestOk = true;
+			std::swap(SourceDirectoryManifest, *Manifest);
 		}
 		else
 		{
@@ -1281,6 +1282,8 @@ SyncDirectory(const FSyncDirectoryOptions& SyncOptions)
 			UNSYNC_BREAK_ON_ERROR;
 			return false;
 		}
+
+		ProxyPool.Dealloc(std::move(Proxy)); // TODO: RAII helper for pooled proxy connections
 	}
 	else if (!SyncOptions.SourceManifestOverride.empty())
 	{
@@ -1894,7 +1897,11 @@ SyncDirectory(const FSyncDirectoryOptions& SyncOptions)
 				if (!GDryRun)
 				{
 					BaseFile.Close();
-					SetFileMtime(Item.TargetFilePath, Item.SourceManifest->Mtime);
+					if (Item.SourceManifest->Mtime)
+					{
+						SetFileMtime(Item.TargetFilePath, Item.SourceManifest->Mtime);
+					}
+
 					if (Item.SourceManifest->bReadOnly)
 					{
 						SetFileReadOnly(Item.TargetFilePath, true);
diff --git a/Engine/Source/Programs/Unsync/Private/UnsyncCore.h b/Engine/Source/Programs/Unsync/Private/UnsyncCore.h
--- a/Engine/Source/Programs/Unsync/Private/UnsyncCore.h
+++ b/Engine/Source/Programs/Unsync/Private/UnsyncCore.h
@@ -226,7 +226,7 @@ FFileSyncResult SyncFile(const FPath&			 SourceFilePath,
 
 struct FSyncDirectoryOptions
 {
-	ESourceType		   SourceType;
+	ESourceType		   SourceType = (ESourceType)0;
 	FPath			   Source;			   // remote data location
 	FPath			   Target;			   // output target location
 	FPath			   Base;			   // base data location, which typically is the same as sync target
diff --git a/Engine/Source/Programs/Unsync/Private/UnsyncFile.cpp b/Engine/Source/Programs/Unsync/Private/UnsyncFile.cpp
--- a/Engine/Source/Programs/Unsync/Private/UnsyncFile.cpp
+++ b/Engine/Source/Programs/Unsync/Private/UnsyncFile.cpp
@@ -964,6 +964,7 @@ bool
 SetFileMtime(const FPath& Path, uint64 Mtime, bool bAllowInDryRun)
 {
 	UNSYNC_ASSERT(!GDryRun || bAllowInDryRun);
+	UNSYNC_ASSERT(Mtime != 0);
 
 	FPath ExtendedPath = MakeExtendedAbsolutePath(Path);
 
diff --git a/Engine/Source/Programs/Unsync/Private/UnsyncHash.cpp b/Engine/Source/Programs/Unsync/Private/UnsyncHash.cpp
--- a/Engine/Source/Programs/Unsync/Private/UnsyncHash.cpp
+++ b/Engine/Source/Programs/Unsync/Private/UnsyncHash.cpp
@@ -142,4 +142,58 @@ const uint32 FBuzHash::TABLE[256] = {
 	0xf9c18d66, 0x593ade65, 0xd95ddf11,
 };
 
+inline bool
+HexCharToUint8(char C, uint8& Output)
+{
+	if (C >= '0' && C <= '9')
+	{
+		Output = C - '0';
+	}
+	else if (C >= 'a' && C <= 'f')
+	{
+		Output = 10 + (C - 'a');
+	}
+	else if (C >= 'A' && C <= 'F')
+	{
+		Output = 10 + (C - 'A');
+	}
+	else
+	{
+		return false;
+	}
+
+	return true;
+}
+
+bool
+ParseHashFromHexString(EStrongHashAlgorithmID StrongHasher, std::string_view HexString, FGenericHash& Output)
+{
+	EHashType HashType = ToHashType(StrongHasher);
+
+	FGenericHash Result = {};
+	Result.Type			= HashType;
+
+	if (Result.Size() != HexString.length() / 2)
+	{
+		return false;
+	}
+
+	const uint64 NumOutputBytes = Result.Size();
+
+	for (uint64 Index = 0; Index < NumOutputBytes; ++Index)
+	{
+		uint8 L = 0, H = 0;
+		bool  bValid = HexCharToUint8(HexString[Index * 2 + 0], H) && HexCharToUint8(HexString[Index * 2 + 1], L);
+		if (!bValid)
+		{
+			return false;
+		}
+		Result.Data[Index] = (H << 4) | L;
+	}
+
+	Output = Result;
+
+	return true;
+}
+
 }  // namespace unsync
diff --git a/Engine/Source/Programs/Unsync/Private/UnsyncHash.h b/Engine/Source/Programs/Unsync/Private/UnsyncHash.h
--- a/Engine/Source/Programs/Unsync/Private/UnsyncHash.h
+++ b/Engine/Source/Programs/Unsync/Private/UnsyncHash.h
@@ -150,6 +150,8 @@ struct FGenericHash
 	};
 };
 
+bool ParseHashFromHexString(EStrongHashAlgorithmID StrongHasher, std::string_view HexString, FGenericHash& Output);
+
 inline FHash160
 ToHash160(const FHash256& Other)
 {
diff --git a/Engine/Source/Programs/Unsync/Private/UnsyncHorde.cpp b/Engine/Source/Programs/Unsync/Private/UnsyncHorde.cpp
new file mode 100644
--- /dev/null
+++ b/Engine/Source/Programs/Unsync/Private/UnsyncHorde.cpp
@@ -0,0 +1,294 @@
+// Copyright Epic Games, Inc. All Rights Reserved.
+
+#include "UnsyncHorde.h"
+#include "UnsyncAuth.h"
+
+#include <regex>
+#include <json11.hpp>
+
+namespace unsync {
+
+FHordeProtocolImpl::FHordeProtocolImpl(const FRemoteDesc& InRemoteDesc, const FBlockRequestMap* InRequestMap, FProxyPool& InProxyPool)
+: FRemoteProtocolBase(InRemoteDesc, InRequestMap)
+, ProxyPool(InProxyPool)
+{
+}
+
+FDownloadResult FHordeProtocolImpl::Download(const TArrayView<FNeedBlock> NeedBlocks, const FBlockDownloadCallback& CompletionCallback)
+{
+	UNSYNC_FATAL(L"FHordeProtocolImpl::Download is not implemented");
+	return FDownloadError(EDownloadRetryMode::Abort);
+}
+
+TResult<FDirectoryManifest>
+FHordeProtocolImpl::DownloadManifest(std::string_view ManifestName)
+{
+	FPooledHttpConnection HttpConnection(ProxyPool);
+
+	std::string BearerToken = ProxyPool.GetAccessToken();
+
+	std::string ManifestUrl = fmt::format("/{}/unsync", ManifestName);
+
+	UNSYNC_LOG(L"Downloading manifest from Horde: '%hs'", ManifestUrl.c_str());
+
+	FHttpRequest Request;
+	Request.Url			= ManifestUrl;
+	Request.BearerToken = BearerToken;
+
+	FHttpResponse Response = HttpRequest(HttpConnection, Request);
+
+	if (!Response.Success())
+	{
+		return HttpError(Response.Code);
+	}
+
+	// Ensure response is terminated
+	Response.Buffer.PushBack(0);
+
+	if (Response.ContentType != EHttpContentType::Application_Json)
+	{
+		return AppError("Unexpected manifest encoding");
+	}
+
+	UNSYNC_LOG(L"Decoding manifest ...");
+
+	return DecodeHordeManifestJson((const char*)Response.Buffer.Data(), ManifestName);
+}
+
+bool RequestPathLooksLikeHordeArtifact(std::string_view RequestPath)
+{
+	static const std::regex Pattern("^api\\/v\\d+\\/artifacts\\/[a-fA-F0-9]+$");
+	return std::regex_match(RequestPath.begin(), RequestPath.end(), Pattern);
+}
+
+TResult<ProxyQuery::FHelloResponse>
+FHordeProtocolImpl::QueryHello(FHttpConnection& HttpConnection)
+{
+	ProxyQuery::FHelloResponse Result;
+
+	// https://horde.devtools.epicgames.com/api/v1/server/auth
+
+	const std::string_view Url = "/api/v1/server/auth";
+
+	FHttpResponse Response = HttpRequest(HttpConnection, EHttpMethod::GET, Url);
+
+	if (!Response.Success())
+	{
+		UNSYNC_ERROR(L"Failed to establish connection to Horde server. Error code: %d.", Response.Code);
+		return HttpError(fmt::format("{}:{}{}", HttpConnection.HostAddress.c_str(), HttpConnection.HostPort, Url), Response.Code);
+	}
+
+	using namespace json11;
+	std::string JsonString = std::string(Response.AsStringView());
+
+	std::string JsonErrorString;
+	Json		JsonObject = Json::parse(JsonString, JsonErrorString);
+
+	if (!JsonErrorString.empty())
+	{
+		return AppError(std::string("JSON parse error while connecting to Horde server: ") + JsonErrorString);
+	}
+
+	if (auto& Field = JsonObject["serverUrl"]; Field.is_string())
+	{
+		Result.AuthServerUri = Field.string_value();
+	}
+
+	if (auto& Field = JsonObject["clientId"]; Field.is_string())
+	{
+		Result.AuthClientId = Field.string_value();
+	}
+
+	if (auto& Field = JsonObject["localRedirectUrls"]; Field.is_array())
+	{
+		if (Field.array_items().size() && Field.array_items()[0].is_string())
+		{
+			// TODO: parse all allowed callback URIs
+			Result.CallbackUri = Field.array_items()[0].string_value();
+		}
+	}
+
+	return ResultOk(Result);
+}
+
+TResult<FDirectoryManifest> DecodeHordeManifestJson(const char* JsonString, std::string_view ArtifactRoot)
+{
+	using namespace json11;
+	std::string JsonErrorString;
+	Json		JsonObject = Json::parse(JsonString, JsonErrorString);
+
+	if (!JsonErrorString.empty())
+	{
+		return AppError(std::string("JSON parse error while downloading manifest from Horde server: ") + JsonErrorString);
+	}
+
+	const uint32 DefaultBlockSize = uint32(64_KB); // TODO: get from manifest JSON
+
+	FDirectoryManifest Manifest;
+	Manifest.Version = FDirectoryManifest::VERSION;
+
+	if (JsonObject["type"] != "unsync_manifest")
+	{
+		return AppError("Manifest JSON is expected to have a 'type' string field with 'unsync_manifest' value");
+	}
+
+	if (auto& Field = JsonObject["hash_strong"]; Field.is_string())
+	{
+		const std::string Value = StringToLower(Field.string_value());
+
+		if (Value == "md5")
+		{
+			Manifest.Algorithm.StrongHashAlgorithmId = EStrongHashAlgorithmID::MD5;
+		}
+		else if (Value == "blake3.128")
+		{
+			Manifest.Algorithm.StrongHashAlgorithmId = EStrongHashAlgorithmID::Blake3_128;
+		}
+		else if (Value == "blake3.160" || Value == "iohash")
+		{
+			Manifest.Algorithm.StrongHashAlgorithmId = EStrongHashAlgorithmID::Blake3_160;
+		}
+		else if (Value == "blake3.256")
+		{
+			Manifest.Algorithm.StrongHashAlgorithmId = EStrongHashAlgorithmID::Blake3_256;
+		}
+		else
+		{
+			return AppError(fmt::format("Unsupported strong hash algorithm '{}'", Value));
+		}
+	}
+
+	if (auto& Field = JsonObject["hash_weak"]; Field.is_string())
+	{
+		const std::string Value = StringToLower(Field.string_value());
+
+		if (Value == "buzhash")
+		{
+			Manifest.Algorithm.WeakHashAlgorithmId = EWeakHashAlgorithmID::BuzHash;
+		}
+		else if (Value == "naive")
+		{
+			Manifest.Algorithm.WeakHashAlgorithmId = EWeakHashAlgorithmID::Naive;
+		}
+		else
+		{
+			return AppError(fmt::format("Unsupported weak hash algorithm '{}'", Value));
+		}
+	}
+
+	if (auto& Field = JsonObject["chunking"]; Field.is_string())
+	{
+		const std::string Value = StringToLower(Field.string_value());
+
+		if (Value == "variable")
+		{
+			Manifest.Algorithm.ChunkingAlgorithmId = EChunkingAlgorithmID::VariableBlocks;
+		}
+		else if (Value == "fixed")
+		{
+			Manifest.Algorithm.ChunkingAlgorithmId = EChunkingAlgorithmID::FixedBlocks;
+		}
+		else
+		{
+			return AppError(fmt::format("Unsupported chunking algorithm '{}'", Value));
+		}
+	}
+
+	if (auto& FiledField = JsonObject["files"]; FiledField.is_array())
+	{
+		for (auto& FileObject : FiledField.array_items())
+		{
+			std::string FileNameUtf8 = FileObject["name"].string_value();
+			std::wstring FileName	  = ConvertUtf8ToWide(FileNameUtf8);
+
+			// Don't include the actual native unsync manifest
+			if (FileName.starts_with(L".unsync"))
+			{
+				continue;
+			}
+
+			std::string ArtifactPathUtf8 = fmt::format("/{}/browse/{}", ArtifactRoot, FileNameUtf8);
+
+			FFileManifest FileManifest;
+			FileManifest.BlockSize = DefaultBlockSize;
+			FileManifest.CurrentPath = ConvertUtf8ToWide(ArtifactPathUtf8);
+
+			if (auto& Field = FileObject["size"]; Field.is_number())
+			{
+				FileManifest.Size = uint64(Field.number_value());
+			}
+
+			if (auto& Field = FileObject["mtime"]; Field.is_number())
+			{
+				FileManifest.Mtime = uint64(Field.number_value());
+			}
+
+			if (auto& Field = FileObject["read_only"]; Field.is_bool())
+			{
+				FileManifest.bReadOnly = Field.bool_value();
+			}
+
+			if (auto& BlocksField = FileObject["blocks"]; BlocksField.is_array())
+			{
+				for (auto& BlockObject : BlocksField.array_items())
+				{
+					FGenericBlock Block;
+					Block.Offset = uint64(BlockObject["offset"].number_value());
+					Block.Size = uint32(BlockObject["size"].number_value());
+					Block.HashWeak = uint32(BlockObject["hash_weak"].number_value());
+					const std::string& StrongHashStr = BlockObject["hash_strong"].string_value();
+					if (!ParseHashFromHexString(Manifest.Algorithm.StrongHashAlgorithmId, StrongHashStr, Block.HashStrong))
+					{
+						return AppError(fmt::format("Failed to parse block strong hash '{}'", StrongHashStr));
+					}
+					FileManifest.Blocks.push_back(Block);
+				}
+			}
+
+			Manifest.Files[FileName] = FileManifest;
+		}
+	}
+
+	return ResultOk(std::move(Manifest));
+}
+
+static const char* GTestHordeManifestJson = R"(
+{
+  "type": "unsync_manifest",
+  "hash_strong": "Blake3.160",
+  "chunking": "Variable",
+  "files": [
+    {
+      "name": "hello_world.txt",
+      "read_only": false,
+      "size": 1095,
+      "blocks": [
+        {
+          "offset": 0,
+          "size": 1095,
+          "hash_strong": "1d9f987b21a19769b758f6fc6354808752620d20"
+        },
+        {
+          "offset": 1095,
+          "size": 1024,
+          "hash_strong": "f52b611e85cf46f466aac6eee0f69b87dff37831"
+        }
+      ]
+    }
+  ]
+}
+)";
+
+void TestHordeManifestDecode()
+{
+	UNSYNC_LOG(L"TestHordeManifestDecode()");
+	UNSYNC_LOG_INDENT;
+
+	TResult<FDirectoryManifest> Manifest = DecodeHordeManifestJson(GTestHordeManifestJson, "api/v2/artifacts/12345");
+	if (Manifest.IsError())
+	{
+		LogError(Manifest.GetError());
+	}
+}
+
+}
diff --git a/Engine/Source/Programs/Unsync/Private/UnsyncHorde.h b/Engine/Source/Programs/Unsync/Private/UnsyncHorde.h
new file mode 100644
--- /dev/null
+++ b/Engine/Source/Programs/Unsync/Private/UnsyncHorde.h
@@ -0,0 +1,32 @@
+// Copyright Epic Games, Inc. All Rights Reserved.
+
+#include "UnsyncCommon.h"
+#include "UnsyncProxy.h"
+
+namespace unsync {
+
+class FBlockRequestMap;
+struct FRemoteDesc;
+
+struct FHordeProtocolImpl : FRemoteProtocolBase
+{
+	FHordeProtocolImpl(const FRemoteDesc& InRemoteDesc, const FBlockRequestMap* InRequestMap, FProxyPool& InProxyPool);
+
+	virtual bool IsValid() const { return bValid; }
+	virtual void Invalidate() { bValid = false; }
+
+	virtual FDownloadResult	 Download(const TArrayView<FNeedBlock> NeedBlocks, const FBlockDownloadCallback& CompletionCallback);
+	virtual TResult<FDirectoryManifest> DownloadManifest(std::string_view ManifestName);
+
+
+	static TResult<ProxyQuery::FHelloResponse> QueryHello(FHttpConnection& HttpConnection);
+
+	FProxyPool& ProxyPool;
+	bool bValid = true;
+};
+
+bool RequestPathLooksLikeHordeArtifact(std::string_view RequestPath);
+
+TResult<FDirectoryManifest> DecodeHordeManifestJson(const char* JsonString, std::string_view ArtifactRoot);
+
+}  // namespace unsync
diff --git a/Engine/Source/Programs/Unsync/Private/UnsyncJupiter.cpp b/Engine/Source/Programs/Unsync/Private/UnsyncJupiter.cpp
--- a/Engine/Source/Programs/Unsync/Private/UnsyncJupiter.cpp
+++ b/Engine/Source/Programs/Unsync/Private/UnsyncJupiter.cpp
@@ -715,7 +715,7 @@ FJupiterProtocolImpl::IsValid() const
 	return bConnected;
 }
 
-TResult<FBuffer>
+TResult<FDirectoryManifest>
 FJupiterProtocolImpl::DownloadManifest(std::string_view ManifestName)
 {
 	// Get root manifest reference object
@@ -797,15 +797,25 @@ FJupiterProtocolImpl::DownloadManifest(std::string_view ManifestName)
 						   FBufferView{},
 						   RemoteDesc.HttpHeaders);
 
-	if (Response.Success())
-	{
-		return ResultOk(std::move(Response.Buffer));
-	}
-	else
+	if (!Response.Success())
 	{
 		// TODO: add a helper to quickly map HttpResponse to a TResult
 		return HttpError(L"Failed to download directory manifest from Jupiter.", Response.Code);
 	}
+
+	FDirectoryManifest Manifest;
+
+	FMemReader		   Reader(Response.Buffer);
+	FIOReaderStream	   Stream(Reader);
+	FPath			   EmptyRoot;  // Don't have a sensible path when not using file system as source
+	bool			   bSourceManifestOk = LoadDirectoryManifest(Manifest, EmptyRoot, Stream);
+
+	if (!bSourceManifestOk)
+	{
+		return AppError(L"Failed to deserialize directory manifest from Jupiter.");
+	}
+
+	return ResultOk(std::move(Manifest));
 }
 
 FDownloadResult
diff --git a/Engine/Source/Programs/Unsync/Private/UnsyncJupiter.h b/Engine/Source/Programs/Unsync/Private/UnsyncJupiter.h
--- a/Engine/Source/Programs/Unsync/Private/UnsyncJupiter.h
+++ b/Engine/Source/Programs/Unsync/Private/UnsyncJupiter.h
@@ -35,14 +35,12 @@ TResult<> JupiterCheckAccess(FHttpConnection& Connection, std::string_view Jupit
 
 struct FJupiterProtocolImpl : FRemoteProtocolBase
 {
-	FJupiterProtocolImpl(const FRemoteDesc&		   InSettings,
-						 const FBlockRequestMap*   InRequestMap,
-						 std::string_view		   HttpHeaders);
-	virtual bool			 IsValid() const override;
-	virtual TResult<FBuffer> DownloadManifest(std::string_view ManifestName) override;
-	virtual FDownloadResult	 Download(const TArrayView<FNeedBlock> NeedBlocks, const FBlockDownloadCallback& CompletionCallback) override;
-	virtual void			 Invalidate() override;
-	virtual bool			 Contains(const FDirectoryManifest& Manifest) override;
+	FJupiterProtocolImpl(const FRemoteDesc& InSettings, const FBlockRequestMap* InRequestMap, std::string_view HttpHeaders);
+	virtual bool						IsValid() const override;
+	virtual TResult<FDirectoryManifest> DownloadManifest(std::string_view ManifestName) override;
+	virtual FDownloadResult Download(const TArrayView<FNeedBlock> NeedBlocks, const FBlockDownloadCallback& CompletionCallback) override;
+	virtual void			Invalidate() override;
+	virtual bool			Contains(const FDirectoryManifest& Manifest) override;
 
 	FHttpConnection Connection;
 	std::string		HttpHeaders;  // TODO: store this in the HttpConnection instead
diff --git a/Engine/Source/Programs/Unsync/Private/UnsyncMain.cpp b/Engine/Source/Programs/Unsync/Private/UnsyncMain.cpp
--- a/Engine/Source/Programs/Unsync/Private/UnsyncMain.cpp
+++ b/Engine/Source/Programs/Unsync/Private/UnsyncMain.cpp
@@ -608,7 +608,7 @@ InnerMain(int Argc, char** Argv)
 					RemoteDesc.StorageNamespace = RemoteDesc.StorageNamespace.substr(0, SlashPos);
 				}
 			}
-			else if (RemoteDesc.Protocol == EProtocolFlavor::Unsync)
+			else if (RemoteDesc.Protocol == EProtocolFlavor::Unsync || RemoteDesc.Protocol == EProtocolFlavor::Horde)
 			{
 				bShouldLogin = true; // Try to authenticate by default when source is a valid URL
 
diff --git a/Engine/Source/Programs/Unsync/Private/UnsyncProxy.cpp b/Engine/Source/Programs/Unsync/Private/UnsyncProxy.cpp
--- a/Engine/Source/Programs/Unsync/Private/UnsyncProxy.cpp
+++ b/Engine/Source/Programs/Unsync/Private/UnsyncProxy.cpp
@@ -8,6 +8,7 @@
 #include "UnsyncPool.h"
 #include "UnsyncProgress.h"
 #include "UnsyncScheduler.h"
+#include "UnsyncHorde.h"
 
 #include <atomic>
 #include <json11.hpp>
@@ -29,7 +30,7 @@ struct FUnsyncBaseProtocolImpl : FRemoteProtocolBase
 	{
 	}
 
-	virtual TResult<FBuffer> DownloadManifest(std::string_view ManifestName) override
+	virtual TResult<FDirectoryManifest> DownloadManifest(std::string_view ManifestName) override
 	{
 		return AppError(L"Manifests can't be downloaded from UNSYNC proxy.");
 	};
@@ -48,7 +49,6 @@ struct FUnsyncProtocolImpl : FUnsyncBaseProtocolImpl
 	virtual FDownloadResult	 Download(const TArrayView<FNeedBlock> NeedBlocks, const FBlockDownloadCallback& CompletionCallback) override;
 
 	virtual void Invalidate() override;
-	virtual bool Contains(const FDirectoryManifest& Manifest) override { return true; }	 // TODO: check files on the unsync proxy
 
 	ESocketSecurity GetSocketSecurity() const;
 
@@ -71,7 +71,6 @@ struct FUnsyncHttpProtocolImpl : FUnsyncBaseProtocolImpl
 
 	virtual bool			IsValid() const override { return bValid; }
 	virtual void			Invalidate() override { bValid = false; }
-	virtual bool			Contains(const FDirectoryManifest& Manifest) override { return true; }	// TODO: check files on the unsync proxy
 
 	virtual FDownloadResult Download(const TArrayView<FNeedBlock> NeedBlocks, const FBlockDownloadCallback& CompletionCallback) override;
 
@@ -93,6 +92,11 @@ FProxy::FProxy(FProxyPool&					  ProxyPool,
 		auto Inner	 = new FJupiterProtocolImpl(RemoteDesc, InRequestMap, RemoteDesc.HttpHeaders);
 		ProtocolImpl = std::unique_ptr<FRemoteProtocolBase>(Inner);
 	}
+	else if (RemoteDesc.Protocol == EProtocolFlavor::Horde)
+	{
+		auto Inner	 = new FHordeProtocolImpl(RemoteDesc, InRequestMap, ProxyPool);
+		ProtocolImpl = std::unique_ptr<FRemoteProtocolBase>(Inner);
+	}
 	else if (RemoteDesc.Protocol == EProtocolFlavor::Unsync)
 	{
 		if (GExperimental && ProxyPool.SupportsHttp() && InFeatures.bBlockDownload)
@@ -236,7 +240,7 @@ FProxy::IsValid() const
 	return ProtocolImpl.get() && ProtocolImpl->IsValid();
 }
 
-TResult<FBuffer>
+TResult<FDirectoryManifest>
 FProxy::DownloadManifest(std::string_view ManifestName)
 {
 	if (ProtocolImpl.get())
@@ -616,7 +620,19 @@ ProxyQuery::Hello(const FRemoteDesc& RemoteDesc, const FAuthDesc* OptAuthDesc)
 {
 	FTlsClientSettings TlsSettings = RemoteDesc.GetTlsClientSettings();
 	FHttpConnection	   Connection(RemoteDesc.Host.Address, RemoteDesc.Host.Port, RemoteDesc.TlsRequirement, TlsSettings);
-	return Hello(Connection, OptAuthDesc);
+
+	if (RemoteDesc.Protocol == EProtocolFlavor::Horde)
+	{
+		return FHordeProtocolImpl::QueryHello(Connection);
+	}
+	else if (RemoteDesc.Protocol == EProtocolFlavor::Unsync)
+	{
+		return Hello(Connection, OptAuthDesc);
+	}
+	else
+	{
+		return AppError("Protocol does not support server information query");
+	}
 }
 
 TResult<ProxyQuery::FHelloResponse>
@@ -1233,7 +1249,14 @@ FProxyPool::FProxyPool(const FRemoteDesc& InRemoteDesc, const FAuthDesc* InAuthD
 	else if (RemoteDesc.Protocol == EProtocolFlavor::Jupiter)
 	{
 		Features.bAuthentication = true;
-		Features.bDownloadByHash = true;
+		Features.bManifestDownload = true;
+	}
+	else if (RemoteDesc.Protocol == EProtocolFlavor::Horde)
+	{
+		Features.bAuthentication   = true;
+		Features.bBlockDownload	   = true;
+		Features.bFileDownload	   = true;
+		Features.bManifestDownload = true;
 	}
 }
 
@@ -1289,6 +1312,28 @@ FProxyPool::DeallocHttp(std::unique_ptr<FHttpConnection>&& Connection)
 	}
 }
 
+std::string
+FProxyPool::GetAccessToken()
+{
+	std::string Result;
+
+	if (AuthDesc)
+	{
+		TResult<FAuthToken> AuthTokenResult = Authenticate(*AuthDesc);
+		if (FAuthToken* AuthToken = AuthTokenResult.TryData())
+		{
+			std::swap(Result, AuthToken->Access);
+		}
+		else
+		{
+			LogError(AuthTokenResult.GetError());
+			UNSYNC_FATAL(L"Cannot proceed without a valid authentication token");
+		}
+	}
+
+	return Result;
+}
+
 void
 FProxyPool::Invalidate()
 {
diff --git a/Engine/Source/Programs/Unsync/Private/UnsyncProxy.h b/Engine/Source/Programs/Unsync/Private/UnsyncProxy.h
--- a/Engine/Source/Programs/Unsync/Private/UnsyncProxy.h
+++ b/Engine/Source/Programs/Unsync/Private/UnsyncProxy.h
@@ -98,7 +98,7 @@ struct FRemoteProtocolFeatures
 	bool bAuthentication   = false;
 	bool bDirectoryListing = false;
 	bool bFileDownload	   = false;
-	bool bDownloadByHash   = false;
+	bool bManifestDownload = false;
 	bool bBlockDownload	   = false;
 };
 
@@ -128,12 +128,12 @@ struct FRemoteProtocolBase
 
 	virtual ~FRemoteProtocolBase(){};
 
-	virtual bool Contains(const FDirectoryManifest& Manifest) = 0;
+	virtual bool Contains(const FDirectoryManifest& Manifest) { return true; }
 	virtual bool IsValid() const							  = 0;
 	virtual void Invalidate()								  = 0;
 
-	virtual FDownloadResult	 Download(const TArrayView<FNeedBlock> NeedBlocks, const FBlockDownloadCallback& CompletionCallback) = 0;
-	virtual TResult<FBuffer> DownloadManifest(std::string_view ManifestName)													 = 0;
+	virtual FDownloadResult Download(const TArrayView<FNeedBlock> NeedBlocks, const FBlockDownloadCallback& CompletionCallback) = 0;
+	virtual TResult<FDirectoryManifest> DownloadManifest(std::string_view ManifestName)											= 0;
 
 	const FBlockRequestMap* RequestMap;
 	FRemoteDesc				RemoteDesc;
@@ -155,8 +155,8 @@ public:
 	bool Contains(const FDirectoryManifest& Manifest);
 	bool IsValid() const;
 
-	FDownloadResult	 Download(const TArrayView<FNeedBlock> NeedBlocks, const FBlockDownloadCallback& CompletionCallback);
-	TResult<FBuffer> DownloadManifest(std::string_view ManifestName);
+	FDownloadResult				Download(const TArrayView<FNeedBlock> NeedBlocks, const FBlockDownloadCallback& CompletionCallback);
+	TResult<FDirectoryManifest> DownloadManifest(std::string_view ManifestName);
 
 private:
 	std::unique_ptr<FRemoteProtocolBase> ProtocolImpl;
@@ -174,6 +174,8 @@ public:
 	std::unique_ptr<FHttpConnection> AllocHttp();
 	void							 DeallocHttp(std::unique_ptr<FHttpConnection>&& Connection);
 
+	std::string GetAccessToken();
+
 	bool SupportsHttp() const { return HttpPool.has_value(); }
 
 	void Invalidate();
diff --git a/Engine/Source/Programs/Unsync/Private/UnsyncRemote.cpp b/Engine/Source/Programs/Unsync/Private/UnsyncRemote.cpp
--- a/Engine/Source/Programs/Unsync/Private/UnsyncRemote.cpp
+++ b/Engine/Source/Programs/Unsync/Private/UnsyncRemote.cpp
@@ -2,6 +2,7 @@
 
 #include "UnsyncRemote.h"
 #include "UnsyncCore.h"
+#include "UnsyncHorde.h"
 
 namespace unsync {
 
@@ -110,6 +111,15 @@ FRemoteDesc::FromUrl(std::string_view Url)
 		}
 	}
 
+	const size_t RequestPos = HostAddress.find_first_of('/');
+	if (RequestPos != std::string::npos)
+	{
+		Result.RequestPath = HostAddress.substr(RequestPos + 1);
+		HostAddress		   = HostAddress.substr(0, RequestPos);
+	}
+
+	const bool bRequestLooksLikeHordeArtifact = RequestPathLooksLikeHordeArtifact(Result.RequestPath);
+
 	const size_t NamespacePos = HostAddress.find_last_of('#');
 
 	switch (Transport)
@@ -123,6 +133,10 @@ FRemoteDesc::FromUrl(std::string_view Url)
 			{
 				Result.Protocol = EProtocolFlavor::Jupiter;
 			}
+			if (bRequestLooksLikeHordeArtifact || Scheme.starts_with("horde"))
+			{
+				Result.Protocol = EProtocolFlavor::Horde;
+			}
 			else
 			{
 				Result.Protocol = EProtocolFlavor::Unsync;
@@ -139,13 +153,6 @@ FRemoteDesc::FromUrl(std::string_view Url)
 	uint16		 HostPort = 0;
 	const size_t PortPos  = HostAddress.find_first_of(':');
 
-	const size_t RequestPos = HostAddress.find_first_of('/');
-	if (RequestPos != std::string::npos)
-	{
-		Result.RequestPath = HostAddress.substr(RequestPos + 1);
-		HostAddress		   = HostAddress.substr(0, RequestPos);
-	}
-
 	if (PortPos == std::string::npos)
 	{
 		switch (Transport)
diff --git a/Engine/Source/Programs/Unsync/Private/UnsyncRemote.h b/Engine/Source/Programs/Unsync/Private/UnsyncRemote.h
--- a/Engine/Source/Programs/Unsync/Private/UnsyncRemote.h
+++ b/Engine/Source/Programs/Unsync/Private/UnsyncRemote.h
@@ -20,6 +20,7 @@ enum class EProtocolFlavor {
 	Unknown,
 	Unsync,
 	Jupiter,
+	Horde,
 };
 
 enum class ETransportProtocol {
diff --git a/Engine/Source/Programs/Unsync/Private/UnsyncSource.h b/Engine/Source/Programs/Unsync/Private/UnsyncSource.h
--- a/Engine/Source/Programs/Unsync/Private/UnsyncSource.h
+++ b/Engine/Source/Programs/Unsync/Private/UnsyncSource.h
@@ -11,7 +11,7 @@ enum class ESourceType : uint8 {
 	Unknown,
 	FileSystem,
 	Server,
-	ServerWithManifestHash,
+	ServerWithManifestId,
 };
 
 struct FSourcePath
diff --git a/Engine/Source/Programs/Unsync/Private/UnsyncTest.cpp b/Engine/Source/Programs/Unsync/Private/UnsyncTest.cpp
--- a/Engine/Source/Programs/Unsync/Private/UnsyncTest.cpp
+++ b/Engine/Source/Programs/Unsync/Private/UnsyncTest.cpp
@@ -946,6 +946,12 @@ RunTests(const std::string& Preset)
 		extern void TestFileAttrib();
 		TestFileAttrib();
 	}
+
+	if (Preset == "horde_manifest_decode" || Preset == "all")
+	{
+		extern void TestHordeManifestDecode();
+		TestHordeManifestDecode();
+	}
 }
 
 }  // namespace unsync
diff --git a/Engine/Source/Programs/Unsync/Private/UnsyncUtil.cpp b/Engine/Source/Programs/Unsync/Private/UnsyncUtil.cpp
--- a/Engine/Source/Programs/Unsync/Private/UnsyncUtil.cpp
+++ b/Engine/Source/Programs/Unsync/Private/UnsyncUtil.cpp
@@ -340,6 +340,15 @@ ToString(const FPath& Path)
 	return ToString(PathView);
 }
 
+std::string
+StringToLower(const std::string& Input)
+{
+	std::string Result = Input;
+	std::transform(Result.begin(), Result.end(), Result.begin(), [](int32 C) { return char(::tolower(C)); });
+	return Result;
+}
+
+
 std::wstring
 StringToLower(const std::wstring& Input)
 {
@@ -725,6 +734,8 @@ LooksLikeUrl(std::string_view Str)
 		"https://",
 		"unsync://",
 		"unsync+tls://",
+		"horde+http://",
+		"horde+https://",
 		"unsync+http://",
 		"unsync+https://",
 		"jupiter+http://",
diff --git a/Engine/Source/Programs/Unsync/Private/UnsyncUtil.h b/Engine/Source/Programs/Unsync/Private/UnsyncUtil.h
--- a/Engine/Source/Programs/Unsync/Private/UnsyncUtil.h
+++ b/Engine/Source/Programs/Unsync/Private/UnsyncUtil.h
@@ -154,6 +154,7 @@ ToString(const std::string_view StringView)
 
 std::string ToString(const FPath& Path);
 
+std::string StringToLower(const std::string& Input);
 std::wstring StringToLower(const std::wstring& Input);
 std::wstring StringToUpper(const std::wstring& Input);
 
