unsync - Remove leftover dev comment

diff --git a/Engine/Source/Programs/Unsync/Private/UnsyncHorde.cpp b/Engine/Source/Programs/Unsync/Private/UnsyncHorde.cpp
--- a/Engine/Source/Programs/Unsync/Private/UnsyncHorde.cpp
+++ b/Engine/Source/Programs/Unsync/Private/UnsyncHorde.cpp
@@ -66,8 +66,6 @@ FHordeProtocolImpl::QueryHello(FHttpConnection& HttpConnection)
 {
 	ProxyQuery::FHelloResponse Result;
 
-	// https://horde.devtools.epicgames.com/api/v1/server/auth
-
 	const std::string_view Url = "/api/v1/server/auth";
 
 	FHttpResponse Response = HttpRequest(HttpConnection, EHttpMethod::GET, Url);
