unsync - Add `--protocol` argument to explicitly specify the server protocol flavor to use

diff --git a/Engine/Source/Programs/Unsync/Private/UnsyncCmdLogin.cpp b/Engine/Source/Programs/Unsync/Private/UnsyncCmdLogin.cpp
--- a/Engine/Source/Programs/Unsync/Private/UnsyncCmdLogin.cpp
+++ b/Engine/Source/Programs/Unsync/Private/UnsyncCmdLogin.cpp
@@ -55,7 +55,20 @@ CmdLogin(const FCmdLoginOptions& Options)
 			FHttpConnection Connection = FHttpConnection::CreateDefaultHttps(Options.Remote);
 
 			FHttpRequest Request;
-			Request.Url			   = "/api/v1/login";
+
+			if (Options.Remote.Protocol == EProtocolFlavor::Unsync)
+			{
+				Request.Url = "/api/v1/login";
+			}
+			if (Options.Remote.Protocol == EProtocolFlavor::Horde)
+			{
+				Request.Url = "/api/v1/projects";
+			}
+			else
+			{
+				Request.Url = "/";
+			}
+
 			Request.Method		   = EHttpMethod::GET;
 			Request.BearerToken	   = AuthTokenResult->Access;
 			FHttpResponse Response = HttpRequest(Connection, Request);
diff --git a/Engine/Source/Programs/Unsync/Private/UnsyncMain.cpp b/Engine/Source/Programs/Unsync/Private/UnsyncMain.cpp
--- a/Engine/Source/Programs/Unsync/Private/UnsyncMain.cpp
+++ b/Engine/Source/Programs/Unsync/Private/UnsyncMain.cpp
@@ -77,7 +77,7 @@ InnerMain(int Argc, char** Argv)
 	std::string				 PresetUtf8		= "all";
 	std::string				 ChunkModeUtf8;
 	std::string				 CacertFilenameUtf8;
-	std::string				 ProtocolName = "jupiter";
+	std::string				 ProtocolName;
 	std::string				 HttpHeaderFilenameUtf8;
 	std::string				 QueryStringUtf8;
 	std::vector<std::string> QueryArgsUtf8;
@@ -134,13 +134,17 @@ InnerMain(int Argc, char** Argv)
 		App->add_flag("--insecure", bAllowInsecureTls, "Skip remote server TLS certificate validation")->group(DangerousGroupId);
 	};
 
-	auto AddProxyOptions = [&RemoteAddressUtf8, &bNoProxySelect](CLI::App* App) {
+	auto AddProxyOptions = [&RemoteAddressUtf8, &ProtocolName, &bNoProxySelect](CLI::App* App)
+	{
 		App->add_option("--proxy, --remote, --server",
 						RemoteAddressUtf8,
-						"Download server address ([transport://]address[:port][/request][#namespace])");
+						"Download server address ([protocol+][transport://]address[:port][/request][#namespace])");
 		App->add_flag("--no-proxy-select",
 					  bNoProxySelect,
 					  "Skip automatic server selection and use the exact one specified by command line or environment variable");
+		App->add_option("--protocol", ProtocolName, "Explicitly specify server protocol instead of inferring it from URL")
+			->required(false)
+			->check(CLI::IsMember({"unsync", "jupiter", "horde"}));
 	};
 
 	// Configure hash
@@ -583,11 +587,18 @@ InnerMain(int Argc, char** Argv)
 		PossibleUrl = QueryArgsUtf8[0];
 	}
 
+	
+	EProtocolFlavor ProtocolFlavorHint = EProtocolFlavor::Unknown;
+	if (!ProtocolName.empty())
+	{
+		ProtocolFlavorHint = ProtocolFlavorFromString(ProtocolName);
+	}
+
 	if (RemoteAddressUtf8.empty() && LooksLikeUrl(PossibleUrl) && (Cli.got_subcommand(SubSync) || Cli.got_subcommand(SubQuery)))
 	{
 		// Derive remote server address from source name if explicit --proxy or --remote option is not provided for sync or query
 
-		TResult<FRemoteDesc> ParsedRemoteDesc = FRemoteDesc::FromUrl(PossibleUrl);
+		TResult<FRemoteDesc> ParsedRemoteDesc = FRemoteDesc::FromUrl(PossibleUrl, ProtocolFlavorHint);
 
 		if (ParsedRemoteDesc.IsOk())
 		{
@@ -632,7 +643,8 @@ InnerMain(int Argc, char** Argv)
 	}
 	else
 	{
-		TResult<FRemoteDesc> ParsedRemoteDesc = FRemoteDesc::FromUrl(RemoteAddressUtf8);
+		TResult<FRemoteDesc> ParsedRemoteDesc = FRemoteDesc::FromUrl(RemoteAddressUtf8, ProtocolFlavorHint);
+
 		if (ParsedRemoteDesc.IsOk())
 		{
 			RemoteDesc = *ParsedRemoteDesc;
diff --git a/Engine/Source/Programs/Unsync/Private/UnsyncRemote.cpp b/Engine/Source/Programs/Unsync/Private/UnsyncRemote.cpp
--- a/Engine/Source/Programs/Unsync/Private/UnsyncRemote.cpp
+++ b/Engine/Source/Programs/Unsync/Private/UnsyncRemote.cpp
@@ -32,12 +32,34 @@ ToString(EProtocolFlavor Protocol)
 		case EProtocolFlavor::Unknown:
 			return "Unknown";
 		case EProtocolFlavor::Unsync:
-			return "UNSYNC";
+			return "Unsync";
+		case EProtocolFlavor::Horde:
+			return "Horde";
 		case EProtocolFlavor::Jupiter:
 			return "Jupiter";
 	}
 }
 
+EProtocolFlavor ProtocolFlavorFromString(std::string_view Str)
+{
+	if (UncasedStringEquals(Str, "unsync"))
+	{
+		return EProtocolFlavor::Unsync;
+	}
+	else if (UncasedStringEquals(Str, "horde"))
+	{
+		return EProtocolFlavor::Horde;
+	}
+	else if (UncasedStringEquals(Str, "jupiter"))
+	{
+		return EProtocolFlavor::Jupiter;
+	}
+	else
+	{
+		return EProtocolFlavor::Unknown;
+	}
+}
+
 static bool
 IsValidUrlCharacter(char C)
 {
@@ -70,7 +92,7 @@ IsValidUrl(std::string_view Url)
 }
 
 TResult<FRemoteDesc>
-FRemoteDesc::FromUrl(std::string_view Url)
+FRemoteDesc::FromUrl(std::string_view Url, EProtocolFlavor ProtocolFlavorHint)
 {
 	if (!IsValidUrl(Url))
 	{
@@ -122,26 +144,35 @@ FRemoteDesc::FromUrl(std::string_view Url)
 
 	const size_t NamespacePos = HostAddress.find_last_of('#');
 
-	switch (Transport)
+	if (ProtocolFlavorHint == EProtocolFlavor::Unknown)
 	{
-		default:
-		case ETransportProtocol::Unsync:
-			Result.Protocol = EProtocolFlavor::Unsync;
-			break;
-		case ETransportProtocol::Http:
-			if (NamespacePos != std::string::npos || Scheme.starts_with("jupiter"))
-			{
-				Result.Protocol = EProtocolFlavor::Jupiter;
-			}
-			if (bRequestLooksLikeHordeArtifact || Scheme.starts_with("horde"))
-			{
-				Result.Protocol = EProtocolFlavor::Horde;
-			}
-			else
-			{
+		// Try to guess protocol flavor
+
+		switch (Transport)
+		{
+			default:
+			case ETransportProtocol::Unsync:
 				Result.Protocol = EProtocolFlavor::Unsync;
-			}
-			break;
+				break;
+			case ETransportProtocol::Http:
+				if (NamespacePos != std::string::npos || Scheme.starts_with("jupiter"))
+				{
+					Result.Protocol = EProtocolFlavor::Jupiter;
+				}
+				if (bRequestLooksLikeHordeArtifact || Scheme.starts_with("horde"))
+				{
+					Result.Protocol = EProtocolFlavor::Horde;
+				}
+				else
+				{
+					Result.Protocol = EProtocolFlavor::Unsync;
+				}
+				break;
+		}
+	}
+	else
+	{
+		Result.Protocol = ProtocolFlavorHint;
 	}
 
 	if (NamespacePos != std::string::npos)
diff --git a/Engine/Source/Programs/Unsync/Private/UnsyncRemote.h b/Engine/Source/Programs/Unsync/Private/UnsyncRemote.h
--- a/Engine/Source/Programs/Unsync/Private/UnsyncRemote.h
+++ b/Engine/Source/Programs/Unsync/Private/UnsyncRemote.h
@@ -16,19 +16,19 @@ class FBuffer;
 
 static constexpr uint16 UNSYNC_DEFAULT_PORT = 53841;
 
-enum class EProtocolFlavor {
+enum class EProtocolFlavor : uint8 {
 	Unknown,
 	Unsync,
 	Jupiter,
 	Horde,
 };
 
-enum class ETransportProtocol {
+enum class ETransportProtocol : uint8 {
 	Http,
 	Unsync,
 };
 
-EProtocolFlavor ProtocolFlavorFromString(const char* Str);
+EProtocolFlavor ProtocolFlavorFromString(std::string_view Str);
 const char*		ToString(EProtocolFlavor Protocol);
 
 struct FHostAddressAndPort
@@ -70,7 +70,7 @@ struct FRemoteDesc
 
 	bool IsValid() const { return Protocol != EProtocolFlavor::Unknown && Host.IsValid(); }
 
-	static TResult<FRemoteDesc> FromUrl(std::string_view Url);
+	static TResult<FRemoteDesc> FromUrl(std::string_view Url, EProtocolFlavor ProtocolFlavorHint = EProtocolFlavor::Unknown);
 
 	FTlsClientSettings GetTlsClientSettings() const;
 };
diff --git a/Engine/Source/Programs/Unsync/Private/UnsyncUtil.cpp b/Engine/Source/Programs/Unsync/Private/UnsyncUtil.cpp
--- a/Engine/Source/Programs/Unsync/Private/UnsyncUtil.cpp
+++ b/Engine/Source/Programs/Unsync/Private/UnsyncUtil.cpp
@@ -32,6 +32,10 @@
 #	include <fmt/xchar.h>
 #endif
 
+#ifdef __GNUC__
+#	define _strnicmp strncasecmp
+#endif
+
 namespace unsync {
 
 static FBuffer GSystemRootCerts;
@@ -365,6 +369,23 @@ StringToUpper(const std::wstring& Input)
 	return Result;
 }
 
+bool
+StringEquals(const std::string_view A, const std::string_view B, bool bCaseSensitive)
+{
+	if (bCaseSensitive)
+	{
+		return A == B;
+	}
+	else if (A.length() == B.length())
+	{
+		return _strnicmp(A.data(), B.data(), A.length()) == 0;
+	}
+	else
+	{
+		return false;
+	}
+}
+
 std::string
 StringEscape(const std::string_view Input)
 {
diff --git a/Engine/Source/Programs/Unsync/Private/UnsyncUtil.h b/Engine/Source/Programs/Unsync/Private/UnsyncUtil.h
--- a/Engine/Source/Programs/Unsync/Private/UnsyncUtil.h
+++ b/Engine/Source/Programs/Unsync/Private/UnsyncUtil.h
@@ -160,6 +160,14 @@ std::wstring StringToUpper(const std::wstring& Input);
 
 std::string StringEscape(const std::string_view Input);
 
+bool StringEquals(const std::string_view A, const std::string_view B, bool bCaseSensitive = true);
+
+inline bool
+UncasedStringEquals(const std::string_view A, const std::string_view B)
+{
+	return StringEquals(A, B, false);
+}
+
 // Takes a drive-based path (e.g. P:/Foo/Bar) and converts it to universal form (e.g. //server/Foo/Bar), if possible.
 // Otherwise, returns original path.
 FPath GetUniversalPath(const FPath& Path);
