unsync - Basic implementation of Horde blob download

diff --git a/Engine/Source/Programs/Unsync/Private/UnsyncFile.h b/Engine/Source/Programs/Unsync/Private/UnsyncFile.h
--- a/Engine/Source/Programs/Unsync/Private/UnsyncFile.h
+++ b/Engine/Source/Programs/Unsync/Private/UnsyncFile.h
@@ -419,11 +419,12 @@ struct FIOReaderStream
 
 	void Seek(uint64 InOffset)
 	{
-		UNSYNC_ASSERT(InOffset < Inner.GetSize());
+		UNSYNC_ASSERT(InOffset <= Inner.GetSize());
 		Offset = InOffset;
 	}
 
 	uint64 Tell() const { return Offset; }
+	void   Skip(uint64 NumBytes) { Seek(Tell() + NumBytes); }
 
 	bool IsValid() const { return Inner.IsValid(); }
 
@@ -433,6 +434,8 @@ struct FIOReaderStream
 		return Read(&Output, sizeof(T));
 	}
 
+	uint64 RemainingSize() const { return std::max(Offset, Inner.GetSize()) - Offset; }
+
 	FIOReader& Inner;
 	uint64	   Offset = 0;
 };
diff --git a/Engine/Source/Programs/Unsync/Private/UnsyncHorde.cpp b/Engine/Source/Programs/Unsync/Private/UnsyncHorde.cpp
--- a/Engine/Source/Programs/Unsync/Private/UnsyncHorde.cpp
+++ b/Engine/Source/Programs/Unsync/Private/UnsyncHorde.cpp
@@ -2,6 +2,8 @@
 
 #include "UnsyncHorde.h"
 #include "UnsyncAuth.h"
+#include "UnsyncCompression.h"
+#include "UnsyncFile.h"
 
 #include <regex>
 #include <json11.hpp>
@@ -16,8 +18,97 @@ FHordeProtocolImpl::FHordeProtocolImpl(const FRemoteDesc& InRemoteDesc, const FB
 
 FDownloadResult FHordeProtocolImpl::Download(const TArrayView<FNeedBlock> NeedBlocks, const FBlockDownloadCallback& CompletionCallback)
 {
-	UNSYNC_FATAL(L"FHordeProtocolImpl::Download is not implemented");
-	return FDownloadError(EDownloadRetryMode::Abort);
+	std::string RequestJson = FormatBlockRequestJson(*RequestMap, NeedBlocks);
+
+	std::string RequestUrl = fmt::format("/{}/unsync-blobs", ProxyPool.RemoteDesc.RequestPath);
+	//std::string RequestUrl = fmt::format("/{}/unsync-blobs?compress=false", ProxyPool.RemoteDesc.RequestPath);
+
+	const EStrongHashAlgorithmID StrongHasher = RequestMap->GetStrongHasher();
+
+	std::string BearerToken = ProxyPool.GetAccessToken();
+
+	FPooledHttpConnection HttpConnection(ProxyPool);
+
+	FHttpRequest Request;
+	Request.Method			   = EHttpMethod::POST;
+	Request.PayloadContentType = EHttpContentType::Application_Json;
+	Request.Payload.Data	   = reinterpret_cast<const uint8*>(RequestJson.data());
+	Request.Payload.Size	   = RequestJson.length();
+	Request.BearerToken		   = BearerToken;
+	Request.Url				   = RequestUrl;
+
+	FHttpResponse Response = HttpRequest(HttpConnection, Request);
+
+	if (!Response.Success())
+	{
+		UNSYNC_ERROR(L"Failed to complete block request. HTTP error code: %d.", Response.Code);
+		return FDownloadError(EDownloadRetryMode::Abort);
+	}
+
+	//std::string_view TransferEncoding = Response.FindHeader("transfer-encoding");
+	const std::string_view ChunkContentEncoding = Response.FindHeader("x-chunk-content-encoding");
+	const std::string_view ContentType			= Response.FindHeader("content-type");
+
+	if (!UncasedStringEquals(ContentType, "application/x-horde-unsync-blob"))
+	{
+		std::string Value = std::string(ContentType);
+		UNSYNC_ERROR(L"Got unexpected blob content type header: '%hs'", Value.c_str());
+		return FDownloadError(EDownloadRetryMode::Abort);
+	}
+
+	// TODO: read body stream as it arrives using HTTP chunk callbacks
+	FMemReader BufferReader(Response.Buffer);
+	FIOReaderStream Reader(BufferReader);
+
+	while (Reader.RemainingSize())
+	{
+		FHordeUnsyncBlobHeaderV1 BlobHeader = {};
+		Reader.ReadInto(BlobHeader.Magic);
+		Reader.ReadInto(BlobHeader.PayloadSize);
+		Reader.ReadInto(BlobHeader.DecompressedSize);
+		Reader.ReadInto(BlobHeader.DecompressedHash);
+
+		if (BlobHeader.Magic != FHordeUnsyncBlobHeaderV1::MAGIC)
+		{
+			UNSYNC_ERROR(L"Got unexpected blob header identifier");
+			return FDownloadError(EDownloadRetryMode::Abort);
+		}
+
+		FDownloadedBlock DownloadedBlock;
+		DownloadedBlock.bCompressed = false;
+
+		FBufferView Payload = Response.Buffer.View(Reader.Tell(), BlobHeader.PayloadSize);
+		FBuffer		DecompressedBuffer;
+
+		if (ChunkContentEncoding == "zstd")
+		{
+			DecompressedBuffer				 = Decompress(Payload.Data, Payload.Size);
+			DownloadedBlock.Data			 = DecompressedBuffer.Data();
+			DownloadedBlock.DecompressedSize = DecompressedBuffer.Size();
+		}
+		else if (ChunkContentEncoding == "" || ChunkContentEncoding == "identity")
+		{
+			UNSYNC_ASSERT(BlobHeader.DecompressedSize == BlobHeader.PayloadSize);
+			DownloadedBlock.Data			 = Payload.Data;
+			DownloadedBlock.DecompressedSize = Payload.Size;
+		}
+		else
+		{
+			std::string Value = std::string(ChunkContentEncoding);
+			UNSYNC_FATAL(L"Unexpected chunk content encoding: '%hs'", Value.c_str());
+		}
+
+		FGenericHash BlockHash	  = ComputeHash(DownloadedBlock.Data, DownloadedBlock.DecompressedSize, StrongHasher);
+		UNSYNC_ASSERT(BlockHash.ToHash160() == BlobHeader.DecompressedHash);
+
+		FHash128	 BlockHash128 = BlockHash.ToHash128();
+
+		CompletionCallback(DownloadedBlock, BlockHash128);
+
+		Reader.Skip(BlobHeader.PayloadSize);
+	}
+
+	return ResultOk<FDownloadError>();
 }
 
 TResult<FDirectoryManifest>
@@ -186,6 +277,11 @@ TResult<FDirectoryManifest> DecodeHordeManifestJson(const char* JsonString, std:
 		{
 			Manifest.Algorithm.ChunkingAlgorithmId = EChunkingAlgorithmID::FixedBlocks;
 		}
+		else if (Value == "rollingbuzhash")
+		{
+			Manifest.Algorithm.ChunkingAlgorithmId = EChunkingAlgorithmID::VariableBlocks;
+			Manifest.Algorithm.WeakHashAlgorithmId = EWeakHashAlgorithmID::BuzHash;
+		}
 		else
 		{
 			return AppError(fmt::format("Unsupported chunking algorithm '{}'", Value));
@@ -205,11 +301,9 @@ TResult<FDirectoryManifest> DecodeHordeManifestJson(const char* JsonString, std:
 				continue;
 			}
 
-			std::string ArtifactPathUtf8 = fmt::format("/{}/browse/{}", ArtifactRoot, FileNameUtf8);
-
 			FFileManifest FileManifest;
 			FileManifest.BlockSize = DefaultBlockSize;
-			FileManifest.CurrentPath = ConvertUtf8ToWide(ArtifactPathUtf8);
+			FileManifest.CurrentPath = FileName;
 
 			if (auto& Field = FileObject["size"]; Field.is_number())
 			{
diff --git a/Engine/Source/Programs/Unsync/Private/UnsyncHttp.cpp b/Engine/Source/Programs/Unsync/Private/UnsyncHttp.cpp
--- a/Engine/Source/Programs/Unsync/Private/UnsyncHttp.cpp
+++ b/Engine/Source/Programs/Unsync/Private/UnsyncHttp.cpp
@@ -38,6 +38,18 @@ UpdateView(std::string_view& View, const char* Data, size_t Size)
 	}
 }
 
+std::string_view FHttpResponse::FindHeader(const std::string_view Name) const
+{
+	for (const auto& It : Headers)
+	{
+		if (UncasedStringEquals(It.first, Name))
+		{
+			return It.second;
+		}
+	}
+	return {};
+}
+
 struct FHttpParser
 {
 	static FHttpParser* ToThis(http_parser* Parser) { return (FHttpParser*)(Parser->data); }
@@ -159,6 +171,9 @@ struct FHttpParser
 				Response.ContentType = EHttpContentType::Text_Plain;
 			}
 		}
+
+		Response.Headers.push_back(std::make_pair(std::string(PendingHeader), std::string(PendingValue)));
+
 		return 0;
 	}
 
diff --git a/Engine/Source/Programs/Unsync/Private/UnsyncHttp.h b/Engine/Source/Programs/Unsync/Private/UnsyncHttp.h
--- a/Engine/Source/Programs/Unsync/Private/UnsyncHttp.h
+++ b/Engine/Source/Programs/Unsync/Private/UnsyncHttp.h
@@ -60,11 +60,15 @@ struct FHttpResponse
 
 	EHttpContentType ContentType = EHttpContentType::Unknown;
 
+	std::vector<std::pair<std::string, std ::string>> Headers;
+
 	bool bConnectionEncrypted = false;
 
 	bool Success() const { return Code >= 200 && Code < 300; }
 
 	std::string_view AsStringView() const { return std::string_view((const char*)Buffer.Data(), Buffer.Size()); }
+
+	std::string_view FindHeader(const std::string_view Name) const;
 };
 
 struct FHttpConnection
diff --git a/Engine/Source/Programs/Unsync/Private/UnsyncProtocol.h b/Engine/Source/Programs/Unsync/Private/UnsyncProtocol.h
--- a/Engine/Source/Programs/Unsync/Private/UnsyncProtocol.h
+++ b/Engine/Source/Programs/Unsync/Private/UnsyncProtocol.h
@@ -290,4 +290,14 @@ struct FBlockPacket256
 	FBuffer	 CompressedData;
 };
 
+struct FHordeUnsyncBlobHeaderV1
+{
+	static constexpr uint64 MAGIC = 0x4C5C2AABA992610Cull;
+
+	uint64 Magic = 0;
+	uint64 PayloadSize = 0;
+	uint64 DecompressedSize = 0;
+	FHash160 DecompressedHash = {};
+};
+
 }  // namespace unsync
diff --git a/Engine/Source/Programs/Unsync/Private/UnsyncProxy.cpp b/Engine/Source/Programs/Unsync/Private/UnsyncProxy.cpp
--- a/Engine/Source/Programs/Unsync/Private/UnsyncProxy.cpp
+++ b/Engine/Source/Programs/Unsync/Private/UnsyncProxy.cpp
@@ -272,12 +272,13 @@ FUnsyncProtocolImpl::IsValid() const
 	return bIsConnetedToHost && SocketHandle && SocketValid(*SocketHandle);
 }
 
-static void
-SortBlockRequests(std::vector<FBlockRequest>& Requests)
+template<typename RequestType>
+void
+SortBlockRequestsByFileName(std::vector<RequestType>& Requests)
 {
 	std::sort(Requests.begin(),
 			  Requests.end(),
-			  [](const FBlockRequest& A, const FBlockRequest& B) -> bool
+			  [](const RequestType& A, const RequestType& B) -> bool
 			  {
 				  int32 FileCmp = std::memcmp(A.FilenameMd5.Data, B.FilenameMd5.Data, A.FilenameMd5.Size());
 				  if (FileCmp != 0)
@@ -322,7 +323,7 @@ FUnsyncProtocolImpl::Download(const TArrayView<FNeedBlock> NeedBlocks, const FBl
 		}
 	}
 
-	SortBlockRequests(Requests);
+	SortBlockRequestsByFileName(Requests);
 
 	bool bOk = bIsConnetedToHost;
 
@@ -466,22 +467,31 @@ FUnsyncProtocolImpl::Download(const TArrayView<FNeedBlock> NeedBlocks, const FBl
 	return ResultOk<FDownloadError>();
 }
 
-static std::string
+std::string
 FormatBlockRequestJson(const FBlockRequestMap& RequestMap, const TArrayView<FNeedBlock> NeedBlocks)
 {
 	const char* StrongHashAlgorithm = ToString(RequestMap.GetStrongHasher());
 
-	std::vector<FBlockRequest> Requests;
+	struct FBlockRequestAndHash : FBlockRequest
+	{
+		FGenericHash FullHash;
+	};
+
+	std::vector<FBlockRequestAndHash> Requests;
 
 	for (const FNeedBlock& Block : NeedBlocks)
 	{
 		if (const FBlockRequest* Request = RequestMap.FindRequest(Block.Hash))
 		{
-			Requests.push_back(*Request);
+			FBlockRequestAndHash Item;
+			static_cast<FBlockRequest&>(Item) = *Request;
+			Item.FullHash					  = Block.Hash;
+
+			Requests.push_back(Item);
 		}
 	}
 
-	SortBlockRequests(Requests);
+	SortBlockRequestsByFileName(Requests);
 
 	std::string Output;
 
@@ -495,7 +505,7 @@ FormatBlockRequestJson(const FBlockRequestMap& RequestMap, const TArrayView<FNee
 	FHash128	   FilenameHash = InvalidHash;
 
 	uint32 BlockIndex = 0;
-	for (const FBlockRequest& Request : Requests)
+	for (const FBlockRequestAndHash& Request : Requests)
 	{
 		if (FilenameHash != Request.FilenameMd5)
 		{
@@ -524,7 +534,7 @@ FormatBlockRequestJson(const FBlockRequestMap& RequestMap, const TArrayView<FNee
 		}
 
 		FGenericBlock Block;
-		Block.HashStrong = FGenericHash::FromBlake3_128(Request.BlockHash);
+		Block.HashStrong = Request.FullHash;
 		Block.Offset	 = Request.Offset;
 		Block.Size		 = CheckedNarrow(Request.Size);
 
@@ -1319,7 +1329,7 @@ FProxyPool::GetAccessToken()
 
 	if (AuthDesc)
 	{
-		TResult<FAuthToken> AuthTokenResult = Authenticate(*AuthDesc);
+		TResult<FAuthToken> AuthTokenResult = Authenticate(*AuthDesc, 120);
 		if (FAuthToken* AuthToken = AuthTokenResult.TryData())
 		{
 			std::swap(Result, AuthToken->Access);
diff --git a/Engine/Source/Programs/Unsync/Private/UnsyncProxy.h b/Engine/Source/Programs/Unsync/Private/UnsyncProxy.h
--- a/Engine/Source/Programs/Unsync/Private/UnsyncProxy.h
+++ b/Engine/Source/Programs/Unsync/Private/UnsyncProxy.h
@@ -307,4 +307,7 @@ struct FRemoteFileSystem : public FProxyFileSystem
 	FProxyPool& ProxyPool;
 };
 
+// Build request block batch using Horde/Unsync JSON request format
+std::string FormatBlockRequestJson(const FBlockRequestMap& RequestMap, const TArrayView<FNeedBlock> NeedBlocks);
+
 }  // namespace unsync
