unsync - Fix directory separators when syncing from Horde

* Also handle some more errors that may be encountered during unsuccessful sync

diff --git a/Engine/Source/Programs/Unsync/Private/UnsyncCmdSync.cpp b/Engine/Source/Programs/Unsync/Private/UnsyncCmdSync.cpp
--- a/Engine/Source/Programs/Unsync/Private/UnsyncCmdSync.cpp
+++ b/Engine/Source/Programs/Unsync/Private/UnsyncCmdSync.cpp
@@ -198,6 +198,7 @@ CmdSync(const FCmdSyncOptions& Options)
 			SyncFileOptions.Algorithm			 = Options.Algorithm;
 			SyncFileOptions.BlockSize			 = uint32(64_KB);
 			SyncFileOptions.bValidateTargetFiles = Options.bValidateTargetFiles;
+			SyncFileOptions.SourceType			 = ESourceType::FileSystem;
 
 			return SyncFile(Options.Source, Options.Target, Options.Target, SyncFileOptions).Succeeded() ? 0 : 1;
 		}
diff --git a/Engine/Source/Programs/Unsync/Private/UnsyncCore.cpp b/Engine/Source/Programs/Unsync/Private/UnsyncCore.cpp
--- a/Engine/Source/Programs/Unsync/Private/UnsyncCore.cpp
+++ b/Engine/Source/Programs/Unsync/Private/UnsyncCore.cpp
@@ -309,6 +309,13 @@ SyncFile(const FNeedList&		   NeedList,
 
 	FFileSyncResult Result;
 
+	if (Options.SourceType == ESourceType::Unknown)
+	{
+		Result.Status = EFileSyncStatus::ErrorInvalidParameters;
+		UNSYNC_ERROR(L"Sync source type must be specified");
+		return Result;
+	}
+
 	uint64 NeedFromSource = ComputeSize(NeedList.Source);
 	uint64 NeedFromBase	  = ComputeSize(NeedList.Base);
 	UNSYNC_VERBOSE(L"Need from source %.2f MB, from base: %.2f MB", SizeMb(NeedFromSource), SizeMb(NeedFromBase));
@@ -384,11 +391,21 @@ SyncFile(const FNeedList&		   NeedList,
 
 		LogStatus(TargetFilePath.wstring().c_str(), L"Patching");
 
-		FDeferredOpenReader SourceFile([SourceFilePath, TargetFilePath] {
-			UNSYNC_VERBOSE(L"Opening source file '%ls'", SourceFilePath.wstring().c_str());
-			LogStatus(TargetFilePath.wstring().c_str(), L"Opening source file");
-			return std::unique_ptr<FNativeFile>(new FNativeFile(SourceFilePath, EFileMode::ReadOnlyUnbuffered));
-		});
+		FDeferredOpenReader SourceFile(
+			[SourceFilePath, TargetFilePath, Options]() -> std::unique_ptr<FIOReader>
+			{
+				if (IsFileSystemSource(Options.SourceType))
+				{
+					UNSYNC_VERBOSE(L"Opening source file '%ls'", SourceFilePath.wstring().c_str());
+					LogStatus(TargetFilePath.wstring().c_str(), L"Opening source file");
+					return std::unique_ptr<FNativeFile>(new FNativeFile(SourceFilePath, EFileMode::ReadOnlyUnbuffered));
+				}
+				else
+				{
+					UNSYNC_ERROR(L"Sync source is not directly accessible");
+					return std::unique_ptr<FIOReader>(new FNullReaderWriter(FNullReaderWriter::FInvalid()));
+				}
+			});
 
 		FBuildTargetParams BuildParams;
 		BuildParams.StrongHasher	 = Options.Algorithm.StrongHashAlgorithmId;
@@ -420,10 +437,13 @@ SyncFile(const FNeedList&		   NeedList,
 				TargetFile = std::make_unique<FNativeFile>(TempTargetFilePath, EFileMode::ReadOnlyUnbuffered);
 			}
 
-			if (!ValidateTarget(*TargetFile, NeedList, Options.Algorithm.StrongHashAlgorithmId))
+			if (TargetFileSizeInfo.TotalBytes > 0)
 			{
-				Result.Status = EFileSyncStatus::ErrorValidation;
-				return Result;
+				if (!ValidateTarget(*TargetFile, NeedList, Options.Algorithm.StrongHashAlgorithmId))
+				{
+					Result.Status = EFileSyncStatus::ErrorValidation;
+					return Result;
+				}
 			}
 		}
 
@@ -1882,6 +1902,7 @@ SyncDirectory(const FSyncDirectoryOptions& SyncOptions)
 			SyncFileOptions.BlockCache			 = BlockCache;
 			SyncFileOptions.ScavengeDatabase	 = ScavengeDatabase.get();
 			SyncFileOptions.bValidateTargetFiles = SyncOptions.bValidateTargetFiles;
+			SyncFileOptions.SourceType			 = SyncOptions.SourceType;
 
 			FFileSyncResult SyncResult =
 				SyncFile(Item.NeedList, Item.ResolvedSourceFilePath, SourceBlocks, BaseFile, Item.TargetFilePath, SyncFileOptions);
diff --git a/Engine/Source/Programs/Unsync/Private/UnsyncCore.h b/Engine/Source/Programs/Unsync/Private/UnsyncCore.h
--- a/Engine/Source/Programs/Unsync/Private/UnsyncCore.h
+++ b/Engine/Source/Programs/Unsync/Private/UnsyncCore.h
@@ -168,6 +168,7 @@ enum class EFileSyncStatus
 	ErrorFinalRename,
 	ErrorTargetFileCreate,
 	ErrorBuildTargetFailed,
+	ErrorInvalidParameters,
 };
 
 const wchar_t* ToString(EFileSyncStatus Status);
@@ -211,6 +212,8 @@ struct FSyncFileOptions
 	FScavengeDatabase* ScavengeDatabase = nullptr;
 
 	bool bValidateTargetFiles = true;  // WARNING: turning this off is intended only for testing/profiling
+
+	ESourceType SourceType = (ESourceType)0;
 };
 
 FFileSyncResult SyncFile(const FPath&			   SourceFilePath,
diff --git a/Engine/Source/Programs/Unsync/Private/UnsyncFile.h b/Engine/Source/Programs/Unsync/Private/UnsyncFile.h
--- a/Engine/Source/Programs/Unsync/Private/UnsyncFile.h
+++ b/Engine/Source/Programs/Unsync/Private/UnsyncFile.h
@@ -143,10 +143,17 @@ struct FIOReader : virtual FIOBase
 	virtual uint64 Read(void* Dest, uint64 SourceOffset, uint64 Size) = 0;
 	virtual bool   ReadAsync(uint64 SourceOffset, uint64 Size, uint64 UserData, IOCallback Callback)
 	{
-		FIOBuffer Buffer   = FIOBuffer::Alloc(Size, L"FIOReader::ReadAsync");
-		uint64	  ReadSize = Read(Buffer.GetData(), SourceOffset, Size);
-		Callback(std::move(Buffer), SourceOffset, ReadSize, UserData);
-		return true;
+		if (Size != 0)
+		{
+			FIOBuffer Buffer   = FIOBuffer::Alloc(Size, L"FIOReader::ReadAsync");
+			uint64	  ReadSize = Read(Buffer.GetData(), SourceOffset, Size);
+			Callback(std::move(Buffer), SourceOffset, ReadSize, UserData);
+			return true;
+		}
+		else
+		{
+			return false;
+		}
 	}
 };
 
@@ -339,13 +346,18 @@ struct FMemReaderWriter : FMemReader, FIOReaderWriter
 
 struct FNullReaderWriter : FIOReaderWriter
 {
-	FNullReaderWriter(uint64 InDataSize) : DataSize(InDataSize) {}
+	struct FInvalid
+	{
+	};
+
+	explicit FNullReaderWriter(uint64 InDataSize) : DataSize(InDataSize) {}
+	explicit FNullReaderWriter(FInvalid) : DataSize(0), bValid(false) {}
 
 	// IOBase
 	virtual void   FlushAll() override{};
 	virtual void   FlushOne() override{};
 	virtual uint64 GetSize() override { return DataSize; }
-	virtual bool   IsValid() override { return true; }
+	virtual bool   IsValid() override { return bValid; }
 	virtual void   Close() override{};
 	virtual int32  GetError() override { return 0; }
 
@@ -367,6 +379,7 @@ struct FNullReaderWriter : FIOReaderWriter
 	virtual uint64 Write(const void* InData, uint64 DestOffset, uint64 WriteSize) override { return WriteSize; }
 
 	uint64 DataSize;
+	bool   bValid = true;
 };
 
 struct FDeferredOpenReader : FIOReader
diff --git a/Engine/Source/Programs/Unsync/Private/UnsyncHorde.cpp b/Engine/Source/Programs/Unsync/Private/UnsyncHorde.cpp
--- a/Engine/Source/Programs/Unsync/Private/UnsyncHorde.cpp
+++ b/Engine/Source/Programs/Unsync/Private/UnsyncHorde.cpp
@@ -29,6 +29,12 @@ FDownloadResult FHordeProtocolImpl::Download(const TArrayView<FNeedBlock> NeedBl
 
 	FPooledHttpConnection HttpConnection(ProxyPool);
 
+	if (!HttpConnection.IsValid())
+	{
+		UNSYNC_ERROR(L"HTTP connection cannot be used");
+		return FDownloadError(EDownloadRetryMode::Abort);
+	}
+
 	FHttpRequest Request;
 	Request.Method			   = EHttpMethod::POST;
 	Request.PayloadContentType = EHttpContentType::Application_Json;
@@ -301,6 +307,8 @@ TResult<FDirectoryManifest> DecodeHordeManifestJson(const char* JsonString, std:
 				continue;
 			}
 
+			ConvertDirectorySeparatorsToNative(FileName);
+
 			FFileManifest FileManifest;
 			FileManifest.BlockSize = DefaultBlockSize;
 			FileManifest.CurrentPath = FileName;
diff --git a/Engine/Source/Programs/Unsync/Private/UnsyncProxy.cpp b/Engine/Source/Programs/Unsync/Private/UnsyncProxy.cpp
--- a/Engine/Source/Programs/Unsync/Private/UnsyncProxy.cpp
+++ b/Engine/Source/Programs/Unsync/Private/UnsyncProxy.cpp
@@ -256,7 +256,7 @@ FProxy::DownloadManifest(std::string_view ManifestName)
 FDownloadResult
 FProxy::Download(const TArrayView<FNeedBlock> NeedBlocks, const FBlockDownloadCallback& CompletionCallback)
 {
-	if (ProtocolImpl.get())
+	if (ProtocolImpl.get() && ProtocolImpl->IsValid())
 	{
 		return ProtocolImpl->Download(NeedBlocks, CompletionCallback);
 	}
diff --git a/Engine/Source/Programs/Unsync/Private/UnsyncProxy.h b/Engine/Source/Programs/Unsync/Private/UnsyncProxy.h
--- a/Engine/Source/Programs/Unsync/Private/UnsyncProxy.h
+++ b/Engine/Source/Programs/Unsync/Private/UnsyncProxy.h
@@ -215,6 +215,7 @@ struct FPooledHttpConnection
 	FHttpConnection& operator*() { return *Get(); }
 	FHttpConnection* operator->() { return Get(); }
 	operator FHttpConnection&() { return *Get(); }
+	bool							 IsValid() const { return ProxyPool.IsValid() && Inner.get(); }
 	FProxyPool&						 ProxyPool;
 	std::unique_ptr<FHttpConnection> Inner;
 };
diff --git a/Engine/Source/Programs/Unsync/Private/UnsyncSource.h b/Engine/Source/Programs/Unsync/Private/UnsyncSource.h
--- a/Engine/Source/Programs/Unsync/Private/UnsyncSource.h
+++ b/Engine/Source/Programs/Unsync/Private/UnsyncSource.h
@@ -14,6 +14,12 @@ enum class ESourceType : uint8 {
 	ServerWithManifestId,
 };
 
+inline bool
+IsFileSystemSource(ESourceType SourceType)
+{
+	return SourceType == ESourceType::FileSystem;
+}
+
 struct FSourcePath
 {
 	ESourceType Type = ESourceType::Unknown;
