unsync - Use current time if file in the manifest doesn't have a valid timestamp

diff --git a/Engine/Source/Programs/Unsync/Private/UnsyncHorde.cpp b/Engine/Source/Programs/Unsync/Private/UnsyncHorde.cpp
--- a/Engine/Source/Programs/Unsync/Private/UnsyncHorde.cpp
+++ b/Engine/Source/Programs/Unsync/Private/UnsyncHorde.cpp
@@ -294,6 +294,10 @@ TResult<FDirectoryManifest> DecodeHordeManifestJson(const char* JsonString, std:
 		}
 	}
 
+	const auto FileTimeNow = std::chrono::time_point<std::chrono::file_clock>::clock::now();
+	const uint64 CurrentWindowsFileTime = ToWindowsFileTime(FileTimeNow);
+	uint32		 NumInvalidTimestamps	= 0;
+
 	if (auto& FiledField = JsonObject["files"]; FiledField.is_array())
 	{
 		for (auto& FileObject : FiledField.array_items())
@@ -323,6 +327,12 @@ TResult<FDirectoryManifest> DecodeHordeManifestJson(const char* JsonString, std:
 				FileManifest.Mtime = uint64(Field.number_value());
 			}
 
+			if (FileManifest.Mtime == 0)
+			{
+				FileManifest.Mtime = CurrentWindowsFileTime;
+				++NumInvalidTimestamps;
+			}
+
 			if (auto& Field = FileObject["read_only"]; Field.is_bool())
 			{
 				FileManifest.bReadOnly = Field.bool_value();
@@ -349,6 +359,11 @@ TResult<FDirectoryManifest> DecodeHordeManifestJson(const char* JsonString, std:
 		}
 	}
 
+	if (NumInvalidTimestamps != 0)
+	{
+		UNSYNC_WARNING(L"Manifest contains files with invalid time stamps (%d) which were set to current time", NumInvalidTimestamps);
+	}
+
 	return ResultOk(std::move(Manifest));
 }
 
