unsync - Handle 0-sized blocks when downloading data from Horde

diff --git a/Engine/Source/Programs/Unsync/Private/UnsyncBuffer.h b/Engine/Source/Programs/Unsync/Private/UnsyncBuffer.h
--- a/Engine/Source/Programs/Unsync/Private/UnsyncBuffer.h
+++ b/Engine/Source/Programs/Unsync/Private/UnsyncBuffer.h
@@ -65,13 +65,13 @@ public:
 
 	uint8& operator[](size_t I)
 	{
-		UNSYNC_ASSERT(I < BufferSize);
+		UNSYNC_ASSERTF(I < BufferSize, L"Buffer index is out of bounds. Index=%llu, BufferSize=%llu.", llu(I), llu(BufferSize));
 		return Buffer[I];
 	}
 
 	const uint8& operator[](size_t I) const
 	{
-		UNSYNC_ASSERT(I < BufferSize);
+		UNSYNC_ASSERTF(I < BufferSize, L"Buffer index is out of bounds. Index=%llu, BufferSize=%llu", llu(I), llu(BufferSize));
 		return Buffer[I];
 	}
 
diff --git a/Engine/Source/Programs/Unsync/Private/UnsyncCompression.cpp b/Engine/Source/Programs/Unsync/Private/UnsyncCompression.cpp
--- a/Engine/Source/Programs/Unsync/Private/UnsyncCompression.cpp
+++ b/Engine/Source/Programs/Unsync/Private/UnsyncCompression.cpp
@@ -69,9 +69,20 @@ Decompress(const uint8* Data, uint64 DataSize)
 {
 	uint64	DecompressedSizeExpected = ZSTD_getDecompressedSize(Data, DataSize);
 	FBuffer Result(DecompressedSizeExpected);
-	uint64	DecompressedSizeActual = ZSTD_decompress(&Result[0], DecompressedSizeExpected, Data, DataSize);
+
+	if (DecompressedSizeExpected == 0)
+	{
+		UNSYNC_ERROR(L"Compressed zstd frame is invalid!");
+		return Result;
+	}
+
+	uint64 DecompressedSizeActual = ZSTD_decompress(Result.Data(), DecompressedSizeExpected, Data, DataSize);
 	if (DecompressedSizeExpected != DecompressedSizeActual)
 	{
+		UNSYNC_ERROR(L"Failed to decompress zstd frame. DecompressedSizeExpected=%llu, DecompressedSizeActual=%llu.",
+					 llu(DecompressedSizeExpected),
+					 llu(DecompressedSizeActual));
+
 		Result.Clear();
 	}
 	return Result;
diff --git a/Engine/Source/Programs/Unsync/Private/UnsyncHorde.cpp b/Engine/Source/Programs/Unsync/Private/UnsyncHorde.cpp
--- a/Engine/Source/Programs/Unsync/Private/UnsyncHorde.cpp
+++ b/Engine/Source/Programs/Unsync/Private/UnsyncHorde.cpp
@@ -88,24 +88,55 @@ FDownloadResult FHordeProtocolImpl::Download(const TArrayView<FNeedBlock> NeedBl
 
 		if (ChunkContentEncoding == "zstd")
 		{
-			DecompressedBuffer				 = Decompress(Payload.Data, Payload.Size);
-			DownloadedBlock.Data			 = DecompressedBuffer.Data();
-			DownloadedBlock.DecompressedSize = DecompressedBuffer.Size();
+			if (BlobHeader.DecompressedSize != 0)
+			{
+				DecompressedBuffer				 = Decompress(Payload.Data, Payload.Size);
+				DownloadedBlock.Data			 = DecompressedBuffer.Data();
+				DownloadedBlock.DecompressedSize = DecompressedBuffer.Size();
+			}
+			else
+			{
+				DownloadedBlock.Data			 = DecompressedBuffer.Data();
+				DownloadedBlock.DecompressedSize = DecompressedBuffer.Size();
+			}
 		}
 		else if (ChunkContentEncoding == "" || ChunkContentEncoding == "identity")
 		{
-			UNSYNC_ASSERT(BlobHeader.DecompressedSize == BlobHeader.PayloadSize);
 			DownloadedBlock.Data			 = Payload.Data;
 			DownloadedBlock.DecompressedSize = Payload.Size;
+
+			if (BlobHeader.DecompressedSize != BlobHeader.PayloadSize)
+			{
+				UNSYNC_ERROR(L"Received blob size (%llu bytes) does not match expected size (%llu bytes)",
+							 llu(BlobHeader.DecompressedSize),
+							 llu(BlobHeader.PayloadSize));
+				return FDownloadError(EDownloadRetryMode::Abort);
+			}
 		}
 		else
 		{
 			std::string Value = std::string(ChunkContentEncoding);
-			UNSYNC_FATAL(L"Unexpected chunk content encoding: '%hs'", Value.c_str());
+			UNSYNC_ERROR(L"Unexpected chunk content encoding: '%hs'", Value.c_str());
+			return FDownloadError(EDownloadRetryMode::Abort);
+		}
+
+		if (DecompressedBuffer.Size() != BlobHeader.DecompressedSize)
+		{
+			UNSYNC_ERROR(L"Received blob size (%llu bytes) does not match expected size (%llu bytes)",
+						 llu(DecompressedBuffer.Size()),
+						 llu(BlobHeader.DecompressedSize));
+			return FDownloadError(EDownloadRetryMode::Abort);
 		}
 
 		FGenericHash BlockHash	  = ComputeHash(DownloadedBlock.Data, DownloadedBlock.DecompressedSize, StrongHasher);
-		UNSYNC_ASSERT(BlockHash.ToHash160() == BlobHeader.DecompressedHash);
+
+		if (BlockHash.ToHash160() != BlobHeader.DecompressedHash)
+		{
+			std::string BlockHashStr	= HashToHexString(BlockHash.ToHash160());
+			std::string ExpectedHashStr = HashToHexString(BlobHeader.DecompressedHash);
+			UNSYNC_ERROR(L"Received blob hash (%hs) does not match expected size (%hs)", BlockHashStr.c_str(), ExpectedHashStr.c_str());
+			return FDownloadError(EDownloadRetryMode::Abort);
+		}
 
 		FHash128	 BlockHash128 = BlockHash.ToHash128();
 
