unsync - Propagate DecodeJwtPayload() error from AcquireAuthToken()

* Fix typo in Horde block validation error message
* Remove auth URL logging

diff --git a/Engine/Source/Programs/Unsync/Private/UnsyncAuth.cpp b/Engine/Source/Programs/Unsync/Private/UnsyncAuth.cpp
--- a/Engine/Source/Programs/Unsync/Private/UnsyncAuth.cpp
+++ b/Engine/Source/Programs/Unsync/Private/UnsyncAuth.cpp
@@ -343,7 +343,7 @@ AcquireAuthToken(const FAuthDesc& AuthDesc, const FOpenIdConfig& OpenIdConfig)
 
 	std::thread ServerThread = StartHttpCallbackServer(CallbackListenSocket, CallbackServerDesc.RequestPath, RandomState, HttpCallbackData);
 
-	UNSYNC_LOG(L"Authorization URL: %hs", AuthorizeUrl.c_str());
+	UNSYNC_LOG(L"Opening authorization URL in default browser");
 	OpenUrlInDefaultBrowser(AuthorizeUrl.c_str());
 
 	UNSYNC_LOG(L"Waiting for HTTP callback on port %d...", int(CallbackPortNumber));
@@ -418,6 +418,10 @@ AcquireAuthToken(const FAuthDesc& AuthDesc, const FOpenIdConfig& OpenIdConfig)
 					Result.ExirationTime = int64(Field.number_value());
 				}
 			}
+			else
+			{
+				return MoveError<FAuthToken>(DecodedAccessTokenResult);
+			}
 
 			Result.Raw = JsonString;
 		}
diff --git a/Engine/Source/Programs/Unsync/Private/UnsyncHorde.cpp b/Engine/Source/Programs/Unsync/Private/UnsyncHorde.cpp
--- a/Engine/Source/Programs/Unsync/Private/UnsyncHorde.cpp
+++ b/Engine/Source/Programs/Unsync/Private/UnsyncHorde.cpp
@@ -134,7 +134,7 @@ FDownloadResult FHordeProtocolImpl::Download(const TArrayView<FNeedBlock> NeedBl
 		{
 			std::string BlockHashStr	= HashToHexString(BlockHash.ToHash160());
 			std::string ExpectedHashStr = HashToHexString(BlobHeader.DecompressedHash);
-			UNSYNC_ERROR(L"Received blob hash (%hs) does not match expected size (%hs)", BlockHashStr.c_str(), ExpectedHashStr.c_str());
+			UNSYNC_ERROR(L"Received blob hash (%hs) does not match expected hash (%hs)", BlockHashStr.c_str(), ExpectedHashStr.c_str());
 			return FDownloadError(EDownloadRetryMode::Abort);
 		}
 
