unsync - Add --no-compression command line argument to prefer uncompressed block downloads when possible

* Fix block size validation when downloading from Horde without compression
* Add more explicit error reporting when attempting to download data from a server without a valid connection pool

diff --git a/Engine/Source/Programs/Unsync/Private/UnsyncCore.cpp b/Engine/Source/Programs/Unsync/Private/UnsyncCore.cpp
--- a/Engine/Source/Programs/Unsync/Private/UnsyncCore.cpp
+++ b/Engine/Source/Programs/Unsync/Private/UnsyncCore.cpp
@@ -297,6 +297,29 @@ ValidateTarget(FIOReader& Reader, const FNeedList& NeedList, EStrongHashAlgorith
 	return NumInvalidBlocks == 0;
 }
 
+static FBuildTargetParams
+GetBuildTargetParams(const FSyncFileOptions& Options)
+{
+	FBuildTargetParams Result;
+
+	Result.StrongHasher		= Options.Algorithm.StrongHashAlgorithmId;
+	Result.ProxyPool		= Options.ProxyPool;
+	Result.BlockCache		= Options.BlockCache;
+	Result.ScavengeDatabase = Options.ScavengeDatabase;
+
+	if (IsFileSystemSource(Options.SourceType))
+	{
+		Result.SourceType = FBuildTargetParams::ESourceType::File;
+	}
+	else
+	{
+		Result.SourceType = FBuildTargetParams::ESourceType::Server;
+	}
+	
+
+	return Result;
+}
+
 FFileSyncResult
 SyncFile(const FNeedList&		   NeedList,
 		 const FPath&			   SourceFilePath,
@@ -407,12 +430,7 @@ SyncFile(const FNeedList&		   NeedList,
 				}
 			});
 
-		FBuildTargetParams BuildParams;
-		BuildParams.StrongHasher	 = Options.Algorithm.StrongHashAlgorithmId;
-		BuildParams.ProxyPool		 = Options.ProxyPool;
-		BuildParams.BlockCache		 = Options.BlockCache;
-		BuildParams.ScavengeDatabase = Options.ScavengeDatabase;
-
+		FBuildTargetParams BuildParams = GetBuildTargetParams(Options);
 		FBuildTargetResult BuildResult = BuildTarget(*TargetFile, SourceFile, BaseDataReader, NeedList, BuildParams);
 
 		Result.SourceBytes = BuildResult.SourceBytes;
diff --git a/Engine/Source/Programs/Unsync/Private/UnsyncHorde.cpp b/Engine/Source/Programs/Unsync/Private/UnsyncHorde.cpp
--- a/Engine/Source/Programs/Unsync/Private/UnsyncHorde.cpp
+++ b/Engine/Source/Programs/Unsync/Private/UnsyncHorde.cpp
@@ -20,8 +20,9 @@ FDownloadResult FHordeProtocolImpl::Download(const TArrayView<FNeedBlock> NeedBl
 {
 	std::string RequestJson = FormatBlockRequestJson(*RequestMap, NeedBlocks);
 
-	std::string RequestUrl = fmt::format("/{}/unsync-blobs", ProxyPool.RemoteDesc.RequestPath);
-	//std::string RequestUrl = fmt::format("/{}/unsync-blobs?compress=false", ProxyPool.RemoteDesc.RequestPath);
+	std::string RequestUrl = fmt::format("/{}/unsync-blobs?compress={}",
+										 ProxyPool.RemoteDesc.RequestPath,
+										 ProxyPool.RemoteDesc.bPreferCompression ? "true" : "false");
 
 	const EStrongHashAlgorithmID StrongHasher = RequestMap->GetStrongHasher();
 
@@ -81,7 +82,7 @@ FDownloadResult FHordeProtocolImpl::Download(const TArrayView<FNeedBlock> NeedBl
 		}
 
 		FDownloadedBlock DownloadedBlock;
-		DownloadedBlock.bCompressed = false;
+		DownloadedBlock.bCompressed = false; // we always decompress the block before handing it over to the caller
 
 		FBufferView Payload = Response.Buffer.View(Reader.Tell(), BlobHeader.PayloadSize);
 		FBuffer		DecompressedBuffer;
@@ -120,15 +121,15 @@ FDownloadResult FHordeProtocolImpl::Download(const TArrayView<FNeedBlock> NeedBl
 			return FDownloadError(EDownloadRetryMode::Abort);
 		}
 
-		if (DecompressedBuffer.Size() != BlobHeader.DecompressedSize)
+		if (DownloadedBlock.DecompressedSize != BlobHeader.DecompressedSize)
 		{
 			UNSYNC_ERROR(L"Received blob size (%llu bytes) does not match expected size (%llu bytes)",
-						 llu(DecompressedBuffer.Size()),
+						 llu(DownloadedBlock.DecompressedSize),
 						 llu(BlobHeader.DecompressedSize));
 			return FDownloadError(EDownloadRetryMode::Abort);
 		}
 
-		FGenericHash BlockHash	  = ComputeHash(DownloadedBlock.Data, DownloadedBlock.DecompressedSize, StrongHasher);
+		FGenericHash BlockHash = ComputeHash(DownloadedBlock.Data, DownloadedBlock.DecompressedSize, StrongHasher);
 
 		if (BlockHash.ToHash160() != BlobHeader.DecompressedHash)
 		{
diff --git a/Engine/Source/Programs/Unsync/Private/UnsyncMain.cpp b/Engine/Source/Programs/Unsync/Private/UnsyncMain.cpp
--- a/Engine/Source/Programs/Unsync/Private/UnsyncMain.cpp
+++ b/Engine/Source/Programs/Unsync/Private/UnsyncMain.cpp
@@ -111,6 +111,7 @@ InnerMain(int Argc, char** Argv)
 	bool					 bNoOutputRevisions  = false;
 	bool					 bPackOnlySmallFiles = false;
 	bool					 bPackFiles		     = false;
+	bool					 bNoCompression		 = false;
 	int32					 CompressionLevel	 = 3;
 	uint32					 DiffBlockSize		 = uint32(4_KB);
 	uint32					 HashOrSyncBlockSize = uint32(64_KB);
@@ -134,7 +135,7 @@ InnerMain(int Argc, char** Argv)
 		App->add_flag("--insecure", bAllowInsecureTls, "Skip remote server TLS certificate validation")->group(DangerousGroupId);
 	};
 
-	auto AddProxyOptions = [&RemoteAddressUtf8, &ProtocolName, &bNoProxySelect](CLI::App* App)
+	auto AddProxyOptions = [&RemoteAddressUtf8, &ProtocolName, &bNoProxySelect, &bNoCompression](CLI::App* App)
 	{
 		App->add_option("--proxy, --remote, --server",
 						RemoteAddressUtf8,
@@ -142,6 +143,9 @@ InnerMain(int Argc, char** Argv)
 		App->add_flag("--no-proxy-select",
 					  bNoProxySelect,
 					  "Skip automatic server selection and use the exact one specified by command line or environment variable");
+		App->add_flag("--no-compression",
+					  bNoCompression,
+					  "Disable compression when downloading blocks from the server, if possible (intended for debugging)");
 		App->add_option("--protocol", ProtocolName, "Explicitly specify server protocol instead of inferring it from URL")
 			->required(false)
 			->check(CLI::IsMember({"unsync", "jupiter", "horde"}));
@@ -658,6 +662,12 @@ InnerMain(int Argc, char** Argv)
 		}
 	}
 
+	if (bNoCompression)
+	{
+		UNSYNC_VERBOSE(L"Uncompressed data transfer is preferred");
+		RemoteDesc.bPreferCompression = false;
+	}
+
 	FPath InputFilename			 = NormalizeFilenameUtf8(InputFilenameUtf8);
 	FPath InputFilename2		 = NormalizeFilenameUtf8(InputFilename2Utf8);
 	FPath OutputFilename		 = NormalizeFilenameUtf8(OutputFilenameUtf8);
diff --git a/Engine/Source/Programs/Unsync/Private/UnsyncRemote.h b/Engine/Source/Programs/Unsync/Private/UnsyncRemote.h
--- a/Engine/Source/Programs/Unsync/Private/UnsyncRemote.h
+++ b/Engine/Source/Programs/Unsync/Private/UnsyncRemote.h
@@ -62,6 +62,8 @@ struct FRemoteDesc
 	std::optional<FHostAddressAndPort> PrimaryHost;	 // Optional address of the server used for login requests and other queries. If
 													 // empty, then HostAddress is used.
 
+	bool bPreferCompression = true;	 // Whether to prefer compressed encoding during bulk data transfer
+
 	const FHostAddressAndPort& GetPrimaryHostAddress() const { return PrimaryHost ? *PrimaryHost : Host; }
 
 	uint32 RecvTimeoutSeconds = 0;
diff --git a/Engine/Source/Programs/Unsync/Private/UnsyncTarget.cpp b/Engine/Source/Programs/Unsync/Private/UnsyncTarget.cpp
--- a/Engine/Source/Programs/Unsync/Private/UnsyncTarget.cpp
+++ b/Engine/Source/Programs/Unsync/Private/UnsyncTarget.cpp
@@ -165,6 +165,21 @@ BuildTarget(FIOWriter& Output, FIOReader& Source, FIOReader& Base, const FNeedLi
 
 	FBuildTargetResult BuildResult;
 
+	if (Params.SourceType == FBuildTargetParams::ESourceType::Server)
+	{
+		if (Params.ProxyPool == nullptr)
+		{
+			UNSYNC_FATAL(L"Connection pool must be provided when syncing from server");
+			return BuildResult;
+		}
+
+		if (!Params.ProxyPool->IsValid())
+		{
+			UNSYNC_FATAL(L"Server connection cannot be established because connection pool is invalid");
+			return BuildResult;
+		}
+	}
+
 	auto TimeBegin = TimePointNow();
 
 	const FNeedListSize			 SizeInfo		  = ComputeNeedListSize(NeedList);
diff --git a/Engine/Source/Programs/Unsync/Private/UnsyncTarget.h b/Engine/Source/Programs/Unsync/Private/UnsyncTarget.h
--- a/Engine/Source/Programs/Unsync/Private/UnsyncTarget.h
+++ b/Engine/Source/Programs/Unsync/Private/UnsyncTarget.h
@@ -29,7 +29,8 @@ struct FBuildTargetParams
 
 	enum class ESourceType {
 		File,
-		Patch
+		Patch,
+		Server,
 	};
 
 	ESourceType SourceType = ESourceType::File;
