unsync - Improve error checking and handling

* Handle most errors during Horde block download as non-fatal (allow other files to be downloaded if one file hits an error)
* Account for block cache creation time when logging individual file download rate
* Add explicit EDownloadRetryMode::Disconnect code to prevent subsequent server requests, distinct from from Abort code
* Propagate errors from few more places when decoding auth tokens

diff --git a/Engine/Source/Programs/Unsync/Private/UnsyncAuth.cpp b/Engine/Source/Programs/Unsync/Private/UnsyncAuth.cpp
--- a/Engine/Source/Programs/Unsync/Private/UnsyncAuth.cpp
+++ b/Engine/Source/Programs/Unsync/Private/UnsyncAuth.cpp
@@ -557,6 +557,10 @@ RefreshAuthToken(const FAuthDesc& AuthDesc, const FOpenIdConfig& OpenIdConfig, c
 					Result.ExirationTime = int64(Field.number_value());
 				}
 			}
+			else
+			{
+				return MoveError<FAuthToken>(DecodedAccessTokenResult);
+			}
 
 			Result.Raw = JsonString;
 		}
@@ -702,6 +706,10 @@ LoadAuthToken(const FPath& Path)
 				AuthToken.ExirationTime = int64(Field.number_value());
 			}
 		}
+		else
+		{
+			return MoveError<FAuthToken>(DecodedAccessTokenResult);
+		}
 
 		return ResultOk(std::move(AuthToken));
 	}
diff --git a/Engine/Source/Programs/Unsync/Private/UnsyncCore.cpp b/Engine/Source/Programs/Unsync/Private/UnsyncCore.cpp
--- a/Engine/Source/Programs/Unsync/Private/UnsyncCore.cpp
+++ b/Engine/Source/Programs/Unsync/Private/UnsyncCore.cpp
@@ -1094,6 +1094,8 @@ struct FFileSyncTaskBatch
 
 	std::unique_ptr<FBlockCache> CreateBlockCache(FProxyPool& ProxyPool, EStrongHashAlgorithmID StrongHasher) const
 	{
+		FTimePoint TimeBegin = TimePointNow();
+
 		std::unique_ptr<FBlockCache> Result = std::make_unique<FBlockCache>();
 
 		Result->BlockData.Resize(NeedBytesFromSource);
@@ -1152,6 +1154,14 @@ struct FFileSyncTaskBatch
 							Result->BlockMap[BlockHash] = FBufferView{OutputView.Data, OutputView.Size};
 							OutputCursor += Block.DecompressedSize;
 						}
+						else
+						{
+							UNSYNC_WARNING(L"Received a corrupt block");
+						}
+					}
+					else
+					{
+						UNSYNC_WARNING(L"Received a block with unexpected hash");
 					}
 				}
 			};
@@ -1159,6 +1169,45 @@ struct FFileSyncTaskBatch
 			FDownloadResult DownloadResult =
 				Proxy->Download(MakeView<FNeedBlock>(UniqueNeedBlocks.data(), UniqueNeedBlocks.size()), DownloadCallback);
 
+			const uint64 NumExpected = UniqueNeedBlocks.size();
+			const uint64 NumDownloaded = Result->BlockMap.size();
+			if (NumExpected != NumDownloaded)
+			{
+				THashSet<FHash128> MissingBlocks = UniqueBlockSet;
+				for (const auto& It : Result->BlockMap)
+				{
+					MissingBlocks.erase(It.first);
+				}
+
+				if (MissingBlocks.size() <= 10)
+				{
+					std::string MissingBlockStr;
+					for (const FHash128& Hash : MissingBlocks)
+					{
+						if (!MissingBlockStr.empty())
+						{
+							MissingBlockStr += ", ";
+						}
+						MissingBlockStr += HashToHexString(Hash);
+					}
+
+					UNSYNC_WARNING(
+						L"Could not download all required data while building block cache. "
+						L"Blocks expected: %llu, actual: %llu. Missing blocks: %hs",
+						llu(NumExpected),
+						llu(NumDownloaded),
+						MissingBlockStr.c_str());
+				}
+				else
+				{
+					UNSYNC_WARNING(
+						L"Could not download all required data while building block cache. "
+						L"Blocks expected: %llu, actual: %llu.",
+						llu(NumExpected),
+						llu(NumDownloaded));
+				}
+			}
+
 			UNSYNC_UNUSED(DownloadResult);
 		}
 
@@ -1166,6 +1215,8 @@ struct FFileSyncTaskBatch
 
 		GScheduler->NetworkSemaphore.Release();
 
+		Result->InitDuration = TimePointNow() - TimeBegin;
+
 		return Result;
 	}
 };
diff --git a/Engine/Source/Programs/Unsync/Private/UnsyncCore.h b/Engine/Source/Programs/Unsync/Private/UnsyncCore.h
--- a/Engine/Source/Programs/Unsync/Private/UnsyncCore.h
+++ b/Engine/Source/Programs/Unsync/Private/UnsyncCore.h
@@ -93,6 +93,7 @@ class FBlockCache
 public:
 	FBuffer							BlockData;
 	THashMap<FHash128, FBufferView> BlockMap;  // Decompressed block data by hash
+	FTimeDuration					InitDuration; // How long it took to create the cache (i.e. download / read all the blocks)
 };
 
 inline uint64
diff --git a/Engine/Source/Programs/Unsync/Private/UnsyncHorde.cpp b/Engine/Source/Programs/Unsync/Private/UnsyncHorde.cpp
--- a/Engine/Source/Programs/Unsync/Private/UnsyncHorde.cpp
+++ b/Engine/Source/Programs/Unsync/Private/UnsyncHorde.cpp
@@ -18,6 +18,11 @@ FHordeProtocolImpl::FHordeProtocolImpl(const FRemoteDesc& InRemoteDesc, const FB
 
 FDownloadResult FHordeProtocolImpl::Download(const TArrayView<FNeedBlock> NeedBlocks, const FBlockDownloadCallback& CompletionCallback)
 {
+	if (NeedBlocks.Size() == 0)
+	{
+		return ResultOk<FDownloadError>();
+	}
+
 	std::string RequestJson = FormatBlockRequestJson(*RequestMap, NeedBlocks);
 
 	std::string RequestUrl = fmt::format("/{}/unsync-blobs?compress={}",
@@ -33,7 +38,7 @@ FDownloadResult FHordeProtocolImpl::Download(const TArrayView<FNeedBlock> NeedBl
 	if (!HttpConnection.IsValid())
 	{
 		UNSYNC_ERROR(L"HTTP connection cannot be used");
-		return FDownloadError(EDownloadRetryMode::Abort);
+		return FDownloadError(EDownloadRetryMode::Disconnect);
 	}
 
 	FHttpRequest Request;
@@ -63,6 +68,12 @@ FDownloadResult FHordeProtocolImpl::Download(const TArrayView<FNeedBlock> NeedBl
 		return FDownloadError(EDownloadRetryMode::Abort);
 	}
 
+	if (Response.Buffer.Size() == 0)
+	{
+		UNSYNC_ERROR(L"Got unexpected empty response body while downloading blocks from Horde");
+		return FDownloadError(EDownloadRetryMode::Abort);
+	}
+
 	// TODO: read body stream as it arrives using HTTP chunk callbacks
 	FMemReader BufferReader(Response.Buffer);
 	FIOReaderStream Reader(BufferReader);
diff --git a/Engine/Source/Programs/Unsync/Private/UnsyncJupiter.cpp b/Engine/Source/Programs/Unsync/Private/UnsyncJupiter.cpp
--- a/Engine/Source/Programs/Unsync/Private/UnsyncJupiter.cpp
+++ b/Engine/Source/Programs/Unsync/Private/UnsyncJupiter.cpp
@@ -823,7 +823,7 @@ FJupiterProtocolImpl::Download(const TArrayView<FNeedBlock> NeedBlocks, const FB
 {
 	if (!IsValid())
 	{
-		return FDownloadError(EDownloadRetryMode::Abort);
+		return FDownloadError(EDownloadRetryMode::Disconnect);
 	}
 
 	if (NeedBlocks.Size() == 0)
@@ -1034,7 +1034,6 @@ FJupiterProtocolImpl::Download(const TArrayView<FNeedBlock> NeedBlocks, const FB
 						|| ErrorsSinceLastSuccess > 10	// Give up after some arbitrary excessive number of errors
 					)
 					{
-						Invalidate();
 						return FDownloadResult(EDownloadRetryMode::Abort);
 					}
 				}
diff --git a/Engine/Source/Programs/Unsync/Private/UnsyncProxy.cpp b/Engine/Source/Programs/Unsync/Private/UnsyncProxy.cpp
--- a/Engine/Source/Programs/Unsync/Private/UnsyncProxy.cpp
+++ b/Engine/Source/Programs/Unsync/Private/UnsyncProxy.cpp
@@ -262,7 +262,7 @@ FProxy::Download(const TArrayView<FNeedBlock> NeedBlocks, const FBlockDownloadCa
 	}
 	else
 	{
-		return FDownloadResult(EDownloadRetryMode::Abort);
+		return FDownloadResult(EDownloadRetryMode::Disconnect);
 	}
 }
 
@@ -294,7 +294,7 @@ FUnsyncProtocolImpl::Download(const TArrayView<FNeedBlock> NeedBlocks, const FBl
 {
 	if (!IsValid())
 	{
-		return FDownloadResult(EDownloadRetryMode::Abort);
+		return FDownloadResult(EDownloadRetryMode::Disconnect);
 	}
 
 	const EStrongHashAlgorithmID StrongHasher = RequestMap->GetStrongHasher();
@@ -390,7 +390,7 @@ FUnsyncProtocolImpl::Download(const TArrayView<FNeedBlock> NeedBlocks, const FBl
 	if (!bOk)
 	{
 		bIsConnetedToHost = false;
-		return FDownloadResult(EDownloadRetryMode::Abort);
+		return FDownloadResult(EDownloadRetryMode::Disconnect);
 	}
 
 	uint64 BytesDownloaded = 0;
@@ -559,7 +559,7 @@ FUnsyncHttpProtocolImpl::Download(const TArrayView<FNeedBlock> NeedBlocks, const
 {
 	if (!IsValid())
 	{
-		return FDownloadResult(EDownloadRetryMode::Abort);
+		return FDownloadResult(EDownloadRetryMode::Disconnect);
 	}
 
 	std::string RequestJson = FormatBlockRequestJson(*RequestMap, NeedBlocks);
@@ -658,6 +658,10 @@ ProxyQuery::Hello(FHttpConnection& HttpConnection, const FAuthDesc* OptAuthDesc)
 		{
 			BearerToken = std::move(AuthTokenResult.GetData().Access);
 		}
+		else
+		{
+			return MoveError<ProxyQuery::FHelloResponse>(AuthTokenResult);
+		}
 	}
 
 	FHttpResponse Response = HttpRequest(HttpConnection, EHttpMethod::GET, Url, {} /*CustomHeaders*/, BearerToken);
diff --git a/Engine/Source/Programs/Unsync/Private/UnsyncProxy.h b/Engine/Source/Programs/Unsync/Private/UnsyncProxy.h
--- a/Engine/Source/Programs/Unsync/Private/UnsyncProxy.h
+++ b/Engine/Source/Programs/Unsync/Private/UnsyncProxy.h
@@ -24,8 +24,9 @@ struct FHttpConnection;
 
 enum class EDownloadRetryMode
 {
-	Retry,	// potentially recoverable error (caller can retry)
-	Abort,	// unrecoverable error
+	Retry,		 // potentially recoverable error (caller can retry the same request)
+	Abort,		 // issuing the same request will likely fail, but other requests may succeed
+	Disconnect,	 // further server API calls are likely to fail
 };
 
 struct FDownloadError : FError
@@ -35,7 +36,7 @@ struct FDownloadError : FError
 
 	EDownloadRetryMode RetryMode = EDownloadRetryMode::Abort;
 
-	bool CanRetry() const { return RetryMode == EDownloadRetryMode ::Retry; }
+	bool CanRetry() const { return RetryMode == EDownloadRetryMode::Retry; }
 };
 
 using FDownloadResult = TResult<FEmpty, FDownloadError>;
diff --git a/Engine/Source/Programs/Unsync/Private/UnsyncTarget.cpp b/Engine/Source/Programs/Unsync/Private/UnsyncTarget.cpp
--- a/Engine/Source/Programs/Unsync/Private/UnsyncTarget.cpp
+++ b/Engine/Source/Programs/Unsync/Private/UnsyncTarget.cpp
@@ -27,10 +27,10 @@ DownloadBlocks(FProxyPool&					  ProxyPool,
 	std::vector<FNeedBlock>		   RemainingNeedBlocks;
 	const std::vector<FNeedBlock>* NeedBlocks = &OriginalUniqueNeedBlocks;
 
-	std::atomic<bool> bGotError = false;
+	FAtomicError Error;
 
 	const uint32 MaxAttempts = 30;
-	for (uint32 Attempt = 0; Attempt <= MaxAttempts && !bGotError; ++Attempt)
+	for (uint32 Attempt = 0; Attempt <= MaxAttempts && !Error; ++Attempt)
 	{
 		if (!ProxyPool.IsValid())
 		{
@@ -59,20 +59,20 @@ DownloadBlocks(FProxyPool&					  ProxyPool,
 
 		const uint64 MaxBytesPerBatch = ProxyPool.RemoteDesc.Protocol == EProtocolFlavor::Jupiter ? 16_MB : 128_MB;
 
-		for (uint64 I = 0; I < NeedBlocks->size() && !bGotError; ++I)
+		for (uint64 BlockIndex = 0; BlockIndex < NeedBlocks->size() && !Error; ++BlockIndex)
 		{
-			const FNeedBlock& Block = (*NeedBlocks)[I];
+			const FNeedBlock& Block = (*NeedBlocks)[BlockIndex];
 
 			if (Batches.back().SizeBytes + Block.Size < MaxBytesPerBatch)
 			{
-				Batches.back().End = I + 1;
+				Batches.back().End = BlockIndex + 1;
 			}
 			else
 			{
-				UNSYNC_ASSERT(Batches.back().End == I);
+				UNSYNC_ASSERT(Batches.back().End == BlockIndex);
 				FDownloadBatch NewBatch = {};
-				NewBatch.Begin			= I;
-				NewBatch.End			= I + 1;
+				NewBatch.Begin			= BlockIndex;
+				NewBatch.End			= BlockIndex + 1;
 				NewBatch.SizeBytes		= 0;
 				Batches.push_back(NewBatch);
 			}
@@ -97,9 +97,9 @@ DownloadBlocks(FProxyPool&					  ProxyPool,
 				 &DownloadedBlocks,
 				 &CompletionCallback,
 				 &NumActiveLogThreads,
-				 &bGotError]
+				 &Error]
 				{
-					if (bGotError)
+					if (Error)
 					{
 						return;
 					}
@@ -126,14 +126,19 @@ DownloadBlocks(FProxyPool&					  ProxyPool,
 												CompletionCallback(Block, BlockHash);
 											});
 
-						if (DownloadResult.IsOk())
-						{
-							ProxyPool.Dealloc(std::move(Proxy));
-						}
-						else if (DownloadResult.GetError().RetryMode == EDownloadRetryMode::Abort)
+						ProxyPool.Dealloc(std::move(Proxy));
+
+						if (const FDownloadError* DownloadError = DownloadResult.TryError())
 						{
-							bGotError = true;
-							ProxyPool.Invalidate();
+							if (!DownloadError->CanRetry())
+							{
+								Error.Set(FError(*DownloadError));
+							}
+
+							if (DownloadError->RetryMode == EDownloadRetryMode::Disconnect)
+							{
+								ProxyPool.Invalidate();
+							}
 						}
 					}
 				});
@@ -141,7 +146,7 @@ DownloadBlocks(FProxyPool&					  ProxyPool,
 
 		DownloadTasks.wait();
 
-		if (DownloadedBlocks.size() == OriginalUniqueNeedBlocks.size() || bGotError)
+		if (DownloadedBlocks.size() == OriginalUniqueNeedBlocks.size() || Error)
 		{
 			break;
 		}
@@ -169,18 +174,19 @@ BuildTarget(FIOWriter& Output, FIOReader& Source, FIOReader& Base, const FNeedLi
 	{
 		if (Params.ProxyPool == nullptr)
 		{
-			UNSYNC_FATAL(L"Connection pool must be provided when syncing from server");
+			UNSYNC_ERROR(L"Connection pool must be provided when syncing from server");
 			return BuildResult;
 		}
 
 		if (!Params.ProxyPool->IsValid())
 		{
-			UNSYNC_FATAL(L"Server connection cannot be established because connection pool is invalid");
+			UNSYNC_ERROR(L"Server connection cannot be established because connection pool is invalid");
 			return BuildResult;
 		}
 	}
 
 	auto TimeBegin = TimePointNow();
+	double ElapsedTime = 0;
 
 	const FNeedListSize			 SizeInfo		  = ComputeNeedListSize(NeedList);
 	const EStrongHashAlgorithmID StrongHasher	  = Params.StrongHasher;
@@ -343,9 +349,9 @@ BuildTarget(FIOWriter& Output, FIOReader& Source, FIOReader& Base, const FNeedLi
 			uint64 ReadBytes = 0;
 
 			auto ReadCallback = [&ReadBytes, &Output, &WriteTasks, &Error, &Stats, Block, ListType](FIOBuffer CmdBuffer,
-																													 uint64	   CmdOffset,
-																													 uint64	   CmdReadSize,
-																													 uint64	   CmdUserData)
+																									uint64	  CmdOffset,
+																									uint64	  CmdReadSize,
+																									uint64	  CmdUserData)
 			{
 				WriteTasks.run(
 					[Buffer = MakeShared(std::move(CmdBuffer)), CmdReadSize, Block, &Output, &Error, &Stats, ListType]()
@@ -426,6 +432,8 @@ BuildTarget(FIOWriter& Output, FIOReader& Source, FIOReader& Base, const FNeedLi
 		UNSYNC_VERBOSE(L"Writing blocks from cache");
 		UNSYNC_LOG_INDENT;
 
+		uint64 BytesFromCache = 0;
+
 		for (const FNeedBlock NeedBlock : FilteredSourceNeedList)
 		{
 			auto It = BlockCache->BlockMap.find(NeedBlock.Hash.ToHash128());
@@ -436,6 +444,7 @@ BuildTarget(FIOWriter& Output, FIOReader& Source, FIOReader& Base, const FNeedLi
 				const uint64 WrittenBytes = Output.Write(BlockBuffer.Data, NeedBlock.TargetOffset, BlockBuffer.Size);
 
 				Stats.WrittenBytesFromSource += WrittenBytes;
+				BytesFromCache += WrittenBytes;
 
 				AddGlobalProgress(NeedBlock.Size, EBlockListType::Source);
 			}
@@ -446,6 +455,17 @@ BuildTarget(FIOWriter& Output, FIOReader& Source, FIOReader& Base, const FNeedLi
 										   [BlockCache](const FNeedBlock& Block)
 										   { return BlockCache->BlockMap.find(Block.Hash.ToHash128()) != BlockCache->BlockMap.end(); });
 
+
+		// Add some virtual time cost to account for block cache creation
+		if (BytesFromCache != 0)
+		{
+			const uint64 BlockCacheTotalSize = BlockCache->BlockData.Size();
+			const double CacheUsedFraction	 = double(BytesFromCache) / double(BlockCacheTotalSize);
+			const double EstimatedTimeCost	 = CacheUsedFraction * BlockCache->InitDuration.count();
+
+			ElapsedTime += EstimatedTimeCost;
+		}
+
 		FilteredSourceNeedList.erase(FilterResult, FilteredSourceNeedList.end());
 	}
 
@@ -664,7 +684,16 @@ BuildTarget(FIOWriter& Output, FIOReader& Source, FIOReader& Base, const FNeedLi
 			FilteredSourceNeedList.erase(FilterResult, FilteredSourceNeedList.end());
 		}
 
-		if (!FilteredSourceNeedList.empty())
+		if (FilteredSourceNeedList.empty())
+		{
+			bSourceDataCopyTaskDone = true;
+		}
+		else if (Params.SourceType == FBuildTargetParams::ESourceType::Server)
+		{
+			Error.Set(AppError(L"Could not download all required data from the server"));
+			bSourceDataCopyTaskDone = true;
+		}
+		else
 		{
 			uint64 RemainingBytes = ComputeSize(FilteredSourceNeedList);
 
@@ -676,10 +705,6 @@ BuildTarget(FIOWriter& Output, FIOReader& Source, FIOReader& Base, const FNeedLi
 
 			ProcessNeedList(Source, FilteredSourceNeedList, RemainingBytes, EBlockListType::Source, bSourceDataCopyTaskDone);
 		}
-		else
-		{
-			bSourceDataCopyTaskDone = true;
-		}
 	}
 	else if (SizeInfo.SourceBytes)
 	{
@@ -695,8 +720,8 @@ BuildTarget(FIOWriter& Output, FIOReader& Source, FIOReader& Base, const FNeedLi
 	BackgroundTasks.wait();
 	WriteTasks.wait();
 
-	double Duration = DurationSec(TimeBegin, TimePointNow());
-	UNSYNC_VERBOSE(L"Done in %.3f sec (%.3f MB / sec)", Duration, SizeMb(double(SizeInfo.TotalBytes) / Duration));
+	ElapsedTime += DurationSec(TimeBegin, TimePointNow());
+	UNSYNC_VERBOSE(L"Done in %.3f sec (%.3f MB / sec)", ElapsedTime, SizeMb(double(SizeInfo.TotalBytes) / ElapsedTime));
 
 	if (GLogVerbose)
 	{
diff --git a/Engine/Source/Programs/Unsync/Private/UnsyncVersion.h b/Engine/Source/Programs/Unsync/Private/UnsyncVersion.h
--- a/Engine/Source/Programs/Unsync/Private/UnsyncVersion.h
+++ b/Engine/Source/Programs/Unsync/Private/UnsyncVersion.h
@@ -4,7 +4,7 @@
 
 #include "UnsyncCommon.h"
 
-#define UNSYNC_VERSION_STR "1.0.78"
+#define UNSYNC_VERSION_STR "1.0.79"
 
 namespace unsync {
 
