unsync - Use common implementation of server directory listing for query list and search commands

diff --git a/Engine/Source/Programs/Unsync/Private/UnsyncCmdQuery.cpp b/Engine/Source/Programs/Unsync/Private/UnsyncCmdQuery.cpp
--- a/Engine/Source/Programs/Unsync/Private/UnsyncCmdQuery.cpp
+++ b/Engine/Source/Programs/Unsync/Private/UnsyncCmdQuery.cpp
@@ -168,9 +168,9 @@ CmdQueryMirrors(const FCmdQueryOptions& Options)
 int32
 CmdQueryList(const FCmdQueryOptions& Options)
 {
-	FHttpConnection Connection = FHttpConnection::CreateDefaultHttps(Options.Remote);
+	FHttpConnection						Connection	  = FHttpConnection::CreateDefaultHttps(Options.Remote);
+	TResult<ProxyQuery::FHelloResponse> HelloResponse = ProxyQuery::Hello(Options.Remote.Protocol, Connection);
 
-	TResult<ProxyQuery::FHelloResponse> HelloResponse = ProxyQuery::Hello(Connection);
 	if (HelloResponse.IsError())
 	{
 		UNSYNC_ERROR("Failed establish a handshake with server '%hs'", Options.Remote.Host.Address.c_str());
@@ -193,33 +193,17 @@ CmdQueryList(const FCmdQueryOptions& Options)
 		return -1;
 	}
 
+	TResult<ProxyQuery::FDirectoryListing> ListingResult = ProxyQuery::ListDirectory(Connection, &AuthDesc, Options.Args[0]);
 
-	std::string Url = fmt::format("/api/v1/list?{}", Options.Args[0]);
-
-	FHttpRequest Request;
-	Request.Url			   = Url;
-	Request.Method		   = EHttpMethod::GET;
-	Request.BearerToken	   = AuthToken->Access;
-	FHttpResponse Response = HttpRequest(Connection, Request);
-
-	if (!Response.Success())
+	if (ListingResult.IsError())
 	{
-		LogError(HttpError(Response.Code));
+		LogError(ListingResult.GetError());
 		return -1;
 	}
 
-	Response.Buffer.PushBack(0);
-
-	std::string	 JsonErrorString;
-	json11::Json JsonObject = json11::Json::parse((const char*)Response.Buffer.Data(), JsonErrorString);
-
-	if (!JsonErrorString.empty())
-	{
-		LogError(AppError(fmt::format("JSON error: {}", JsonErrorString.c_str())));
-		return -1;
-	}
+	std::string ListingJson = ListingResult->ToJson();
 
-	LogPrintf(ELogLevel::MachineReadable, L"%hs\n", Response.Buffer.Data());
+	LogPrintf(ELogLevel::MachineReadable, L"%hs\n", ListingJson.c_str());
 
 	return 0;
 }
@@ -243,31 +227,23 @@ CmdQuerySearch(const FCmdQueryOptions& Options)
 
 	TObjectPool<FHttpConnection> ConnectionPool(CreateConnection);
 
-	std::string BearerToken;
-
+	TResult<ProxyQuery::FHelloResponse> HelloResponse = [&ConnectionPool, &Options]
 	{
 		std::unique_ptr<FHttpConnection> Connection = ConnectionPool.Acquire();
-		TResult<ProxyQuery::FHelloResponse> HelloResponse = ProxyQuery::Hello(*Connection);
+		TResult<ProxyQuery::FHelloResponse> Result = ProxyQuery::Hello(Options.Remote.Protocol, *Connection);
 		ConnectionPool.Release(std::move(Connection));
-		if (HelloResponse.IsError())
-		{
-			UNSYNC_ERROR("Failed establish a handshake with server '%hs'", Options.Remote.Host.Address.c_str());
-			LogError(HelloResponse.GetError());
-			return -1;
-		}
-
-		FAuthDesc AuthDesc = FAuthDesc::FromHelloResponse(*HelloResponse);
+		return Result;
+	}();
 
-		TResult<FAuthToken> AuthToken = Authenticate(AuthDesc, 5 * 60);
-		if (!AuthToken.IsOk())
-		{
-			LogError(AuthToken.GetError());
-			return -1;
-		}
-
-		BearerToken = std::move(AuthToken->Access);
+	if (HelloResponse.IsError())
+	{
+		UNSYNC_ERROR("Failed establish a handshake with server '%hs'", Options.Remote.Host.Address.c_str());
+		LogError(HelloResponse.GetError());
+		return -1;
 	}
 
+	FAuthDesc AuthDesc = FAuthDesc::FromHelloResponse(*HelloResponse);
+
 	const std::string& RootPath = Options.Args[0];
 	std::vector<std::regex> SubdirPatterns;
 
@@ -319,37 +295,21 @@ CmdQuerySearch(const FCmdQueryOptions& Options)
 	Context.ParentThreadIndent	 = GLogIndent;
 
 	std::function<void(std::string, int32)> ExploreDirectory =
-		[&Context, &BearerToken, &ConnectionPool, &ExploreDirectory, &SubdirPatterns, &Tasks](std::string Path, int32 CurrentDepth)
+		[&Context, &AuthDesc, &ConnectionPool, &ExploreDirectory, &SubdirPatterns, &Tasks](std::string Path, int32 CurrentDepth)
 	{
 		FLogVerbosityScope VerboseScope(Context.bParentThreadVerbose);
 		FLogIndentScope	   IndentScope(Context.ParentThreadIndent, true);
 
 		UNSYNC_VERBOSE2(L"Listing '%hs'", Path.c_str());
 
-		std::string Url = fmt::format("/api/v1/list?{}", Path);
-
-		FHttpRequest Request;
-		Request.Url			= Url;
-		Request.Method		= EHttpMethod::GET;
-		Request.BearerToken = BearerToken;
-
 		GScheduler->NetworkSemaphore.Acquire(false);
-
 		std::unique_ptr<FHttpConnection> Connection = ConnectionPool.Acquire();
-		FHttpResponse					 Response	= HttpRequest(*Connection, Request);
-		ConnectionPool.Release(std::move(Connection));
-
-		GScheduler->NetworkSemaphore.Release();
 
-		if (!Response.Success())
-		{
-			LogError(HttpError(Response.Code));
-			return;
-		}
+		TResult<FDirectoryListing> DirectoryListingResult = ProxyQuery::ListDirectory(*Connection, &AuthDesc, Path);
 
-		Response.Buffer.PushBack(0);
+		ConnectionPool.Release(std::move(Connection));
+		GScheduler->NetworkSemaphore.Release();
 
-		TResult<FDirectoryListing> DirectoryListingResult = FDirectoryListing::FromJson((const char*)Response.Buffer.Data());
 		if (DirectoryListingResult.IsError())
 		{
 			LogError(DirectoryListingResult.GetError());
@@ -429,8 +389,8 @@ CmdQueryFile(const FCmdQueryOptions& Options)
 	}
 
 	// TODO: use a global connection pool and use it for ProxyQuery::DownloadFile too
-	FHttpConnection HelloConnection = FHttpConnection::CreateDefaultHttps(Options.Remote);
-	TResult<ProxyQuery::FHelloResponse> HelloResponse = ProxyQuery::Hello(HelloConnection);
+	FHttpConnection						HelloConnection = FHttpConnection::CreateDefaultHttps(Options.Remote);
+	TResult<ProxyQuery::FHelloResponse> HelloResponse	= ProxyQuery::Hello(Options.Remote.Protocol, HelloConnection);
 	if (HelloResponse.IsError())
 	{
 		UNSYNC_ERROR("Failed establish a handshake with server '%hs'", Options.Remote.Host.Address.c_str());
diff --git a/Engine/Source/Programs/Unsync/Private/UnsyncProxy.cpp b/Engine/Source/Programs/Unsync/Private/UnsyncProxy.cpp
--- a/Engine/Source/Programs/Unsync/Private/UnsyncProxy.cpp
+++ b/Engine/Source/Programs/Unsync/Private/UnsyncProxy.cpp
@@ -9,6 +9,7 @@
 #include "UnsyncProgress.h"
 #include "UnsyncScheduler.h"
 #include "UnsyncHorde.h"
+#include "UnsyncUtil.h"
 
 #include <atomic>
 #include <json11.hpp>
@@ -56,6 +57,7 @@ struct FUnsyncProtocolImpl : FUnsyncBaseProtocolImpl
 	std::unique_ptr<FSocketBase> SocketHandle;
 
 	static void SendTelemetryEvent(const FRemoteDesc& RemoteDesc, const FTelemetryEventSyncComplete& Event);
+	static TResult<ProxyQuery::FHelloResponse> QueryHello(FHttpConnection& HttpConnection, const FAuthDesc* OptAuthDesc=nullptr);
 };
 
 struct FUnsyncHttpProtocolImpl : FUnsyncBaseProtocolImpl
@@ -631,13 +633,19 @@ ProxyQuery::Hello(const FRemoteDesc& RemoteDesc, const FAuthDesc* OptAuthDesc)
 	FTlsClientSettings TlsSettings = RemoteDesc.GetTlsClientSettings();
 	FHttpConnection	   Connection(RemoteDesc.Host.Address, RemoteDesc.Host.Port, RemoteDesc.TlsRequirement, TlsSettings);
 
-	if (RemoteDesc.Protocol == EProtocolFlavor::Horde)
+	return Hello(RemoteDesc.Protocol, Connection, OptAuthDesc);
+}
+
+TResult<ProxyQuery::FHelloResponse>
+ProxyQuery::Hello(EProtocolFlavor Protocol, FHttpConnection& Connection, const FAuthDesc* OptAuthDesc)
+{
+	if (Protocol == EProtocolFlavor::Horde)
 	{
 		return FHordeProtocolImpl::QueryHello(Connection);
 	}
-	else if (RemoteDesc.Protocol == EProtocolFlavor::Unsync)
+	else if (Protocol == EProtocolFlavor::Unsync)
 	{
-		return Hello(Connection, OptAuthDesc);
+		return FUnsyncProtocolImpl::QueryHello(Connection, OptAuthDesc);
 	}
 	else
 	{
@@ -646,8 +654,10 @@ ProxyQuery::Hello(const FRemoteDesc& RemoteDesc, const FAuthDesc* OptAuthDesc)
 }
 
 TResult<ProxyQuery::FHelloResponse>
-ProxyQuery::Hello(FHttpConnection& HttpConnection, const FAuthDesc* OptAuthDesc)
+FUnsyncProtocolImpl::QueryHello(FHttpConnection& HttpConnection, const FAuthDesc* OptAuthDesc)
 {
+	using ProxyQuery::FHelloResponse;
+
 	const char* Url = "/api/v1/hello";
 
 	std::string BearerToken;
@@ -761,7 +771,7 @@ ProxyQuery::Hello(FHttpConnection& HttpConnection, const FAuthDesc* OptAuthDesc)
 
 	if (auto& Field = JsonObject["primary"]; Field.is_string())
 	{
-		const std::string& PrimaryHostStr = Field.string_value();
+		const std::string&	 PrimaryHostStr	 = Field.string_value();
 		TResult<FRemoteDesc> PrimaryHostDesc = FRemoteDesc::FromUrl(PrimaryHostStr);
 		if (PrimaryHostDesc.IsOk())
 		{
@@ -820,6 +830,36 @@ ProxyQuery::FDirectoryListing::FromJson(const char* JsonString)
 	return ResultOk(std::move(Result));
 }
 
+std::string
+ProxyQuery::FDirectoryListing::ToJson() const
+{
+	std::string Result;
+
+	Result += "{\"entries\": [\n";
+
+	uint64 EntryIndex = 0;
+	for (const FDirectoryListingEntry& Entry : Entries)
+	{
+		if (EntryIndex != 0)
+		{
+			Result += ",\n";
+		}
+
+		Result += "{ ";
+		FormatJsonKeyValueStr(Result, "name", Entry.Name, ", ");
+		FormatJsonKeyValueBool(Result, "is_directory", Entry.bDirectory, ", ");
+		FormatJsonKeyValueUInt(Result, "mtime", Entry.Mtime, ", ");
+		FormatJsonKeyValueUInt(Result, "size", Entry.Size);
+		Result += "}";
+
+		++EntryIndex;
+	}
+
+	Result += "\n]}\n";
+
+	return Result;
+}
+
 TResult<ProxyQuery::FDirectoryListing>
 ProxyQuery::ListDirectory(FHttpConnection& Connection, const FAuthDesc* AuthDesc, const std::string& Path)
 {
@@ -1236,7 +1276,7 @@ FProxyPool::FProxyPool(const FRemoteDesc& InRemoteDesc, const FAuthDesc* InAuthD
 
 		std::unique_ptr<FHttpConnection> HttpConnection = HttpPool->Acquire();
 
-		TResult<ProxyQuery::FHelloResponse> Response = ProxyQuery::Hello(*HttpConnection, AuthDesc);
+		TResult<ProxyQuery::FHelloResponse> Response = ProxyQuery::Hello(RemoteDesc.Protocol, *HttpConnection, AuthDesc);
 
 		HttpPool->Release(std::move(HttpConnection));
 
diff --git a/Engine/Source/Programs/Unsync/Private/UnsyncProxy.h b/Engine/Source/Programs/Unsync/Private/UnsyncProxy.h
--- a/Engine/Source/Programs/Unsync/Private/UnsyncProxy.h
+++ b/Engine/Source/Programs/Unsync/Private/UnsyncProxy.h
@@ -247,7 +247,7 @@ struct FHelloResponse
 	bool SupportsAuthentication() const { return Features.bAuthentication && !AuthServerUri.empty() && !AuthClientId.empty(); }
 };
 TResult<FHelloResponse> Hello(const FRemoteDesc& RemoteDesc, const FAuthDesc* OptAuthDesc = nullptr);
-TResult<FHelloResponse> Hello(FHttpConnection& Connection, const FAuthDesc* OptAuthDesc = nullptr);
+TResult<FHelloResponse> Hello(EProtocolFlavor Protocol, FHttpConnection& Connection, const FAuthDesc* OptAuthDesc = nullptr);
 
 struct FDirectoryListingEntry
 {
@@ -262,6 +262,7 @@ struct FDirectoryListing
 	std::vector<FDirectoryListingEntry> Entries;
 
 	static TResult<FDirectoryListing> FromJson(const char* JsonString);
+	std::string						  ToJson() const;
 };
 
 TResult<FDirectoryListing> ListDirectory(FHttpConnection& Connection, const FAuthDesc* AuthDesc, const std::string& Path);
