unsync - Implement experimental mechanism to list Horde artifacts as a virtual directory structure

* Use 5 minute threshold by default for triggering auth token refresh

diff --git a/Engine/Source/Programs/Unsync/Private/UnsyncAuth.h b/Engine/Source/Programs/Unsync/Private/UnsyncAuth.h
--- a/Engine/Source/Programs/Unsync/Private/UnsyncAuth.h
+++ b/Engine/Source/Programs/Unsync/Private/UnsyncAuth.h
@@ -67,8 +67,8 @@ TResult<FOpenIdConfig> GetOpenIdConfig(const FAuthDesc& AuthDesc);
 // - Attempt use stored refresh token first
 // - If refresh is not possible, use PKCE Authentication flow to get new tokens
 // - Save refresh token in user directory for future use
-// - Skips acquiring new token if remaining valid time is above RefreshThreshold (in seconds)
-TResult<FAuthToken> Authenticate(const FAuthDesc& AuthDesc, int32 RefreshThreshold = INT_MAX);
+// - Skips acquiring new token if remaining valid time is above RefreshThreshold (in seconds, 5 minutes by default)
+TResult<FAuthToken> Authenticate(const FAuthDesc& AuthDesc, int32 RefreshThreshold = 300);
 
 // Auth utility functions
 
diff --git a/Engine/Source/Programs/Unsync/Private/UnsyncCmdQuery.cpp b/Engine/Source/Programs/Unsync/Private/UnsyncCmdQuery.cpp
--- a/Engine/Source/Programs/Unsync/Private/UnsyncCmdQuery.cpp
+++ b/Engine/Source/Programs/Unsync/Private/UnsyncCmdQuery.cpp
@@ -179,21 +179,13 @@ CmdQueryList(const FCmdQueryOptions& Options)
 	}
 	FAuthDesc AuthDesc = FAuthDesc::FromHelloResponse(*HelloResponse);
 
-	TResult<FAuthToken> AuthToken = Authenticate(AuthDesc, 5 * 60);
-
-	if (!AuthToken.IsOk())
-	{
-		LogError(AuthToken.GetError());
-		return -1;
-	}
-
 	if (Options.Args.empty())
 	{
 		UNSYNC_ERROR(L"Path argument is required");
 		return -1;
 	}
 
-	TResult<ProxyQuery::FDirectoryListing> ListingResult = ProxyQuery::ListDirectory(Connection, &AuthDesc, Options.Args[0]);
+	TResult<ProxyQuery::FDirectoryListing> ListingResult = ProxyQuery::ListDirectory(Options.Remote.Protocol, Connection, &AuthDesc, Options.Args[0]);
 
 	if (ListingResult.IsError())
 	{
@@ -295,7 +287,7 @@ CmdQuerySearch(const FCmdQueryOptions& Options)
 	Context.ParentThreadIndent	 = GLogIndent;
 
 	std::function<void(std::string, int32)> ExploreDirectory =
-		[&Context, &AuthDesc, &ConnectionPool, &ExploreDirectory, &SubdirPatterns, &Tasks](std::string Path, int32 CurrentDepth)
+		[&Context, &AuthDesc, &ConnectionPool, &ExploreDirectory, &SubdirPatterns, &Tasks, &Options](std::string Path, int32 CurrentDepth)
 	{
 		FLogVerbosityScope VerboseScope(Context.bParentThreadVerbose);
 		FLogIndentScope	   IndentScope(Context.ParentThreadIndent, true);
@@ -305,7 +297,8 @@ CmdQuerySearch(const FCmdQueryOptions& Options)
 		GScheduler->NetworkSemaphore.Acquire(false);
 		std::unique_ptr<FHttpConnection> Connection = ConnectionPool.Acquire();
 
-		TResult<FDirectoryListing> DirectoryListingResult = ProxyQuery::ListDirectory(*Connection, &AuthDesc, Path);
+		TResult<FDirectoryListing> DirectoryListingResult =
+			ProxyQuery::ListDirectory(Options.Remote.Protocol, *Connection, &AuthDesc, Path);
 
 		ConnectionPool.Release(std::move(Connection));
 		GScheduler->NetworkSemaphore.Release();
@@ -400,7 +393,7 @@ CmdQueryFile(const FCmdQueryOptions& Options)
 	FAuthDesc AuthDesc = FAuthDesc::FromHelloResponse(*HelloResponse);
 	HelloConnection.Close();
 
-	TResult<FAuthToken> AuthToken = Authenticate(AuthDesc, 5 * 60);
+	TResult<FAuthToken> AuthToken = Authenticate(AuthDesc);
 
 	if (!AuthToken.IsOk())
 	{
diff --git a/Engine/Source/Programs/Unsync/Private/UnsyncHorde.cpp b/Engine/Source/Programs/Unsync/Private/UnsyncHorde.cpp
--- a/Engine/Source/Programs/Unsync/Private/UnsyncHorde.cpp
+++ b/Engine/Source/Programs/Unsync/Private/UnsyncHorde.cpp
@@ -7,6 +7,8 @@
 
 #include <regex>
 #include <json11.hpp>
+#include <optional>
+#include <stdio.h>
 
 namespace unsync {
 
@@ -16,7 +18,8 @@ FHordeProtocolImpl::FHordeProtocolImpl(const FRemoteDesc& InRemoteDesc, const FB
 {
 }
 
-FDownloadResult FHordeProtocolImpl::Download(const TArrayView<FNeedBlock> NeedBlocks, const FBlockDownloadCallback& CompletionCallback)
+FDownloadResult
+FHordeProtocolImpl::Download(const TArrayView<FNeedBlock> NeedBlocks, const FBlockDownloadCallback& CompletionCallback)
 {
 	if (NeedBlocks.Size() == 0)
 	{
@@ -57,7 +60,7 @@ FDownloadResult FHordeProtocolImpl::Download(const TArrayView<FNeedBlock> NeedBl
 		return FDownloadError(EDownloadRetryMode::Abort);
 	}
 
-	//std::string_view TransferEncoding = Response.FindHeader("transfer-encoding");
+	// std::string_view TransferEncoding = Response.FindHeader("transfer-encoding");
 	const std::string_view ChunkContentEncoding = Response.FindHeader("x-chunk-content-encoding");
 	const std::string_view ContentType			= Response.FindHeader("content-type");
 
@@ -75,7 +78,7 @@ FDownloadResult FHordeProtocolImpl::Download(const TArrayView<FNeedBlock> NeedBl
 	}
 
 	// TODO: read body stream as it arrives using HTTP chunk callbacks
-	FMemReader BufferReader(Response.Buffer);
+	FMemReader		BufferReader(Response.Buffer);
 	FIOReaderStream Reader(BufferReader);
 
 	while (Reader.RemainingSize())
@@ -93,7 +96,7 @@ FDownloadResult FHordeProtocolImpl::Download(const TArrayView<FNeedBlock> NeedBl
 		}
 
 		FDownloadedBlock DownloadedBlock;
-		DownloadedBlock.bCompressed = false; // we always decompress the block before handing it over to the caller
+		DownloadedBlock.bCompressed = false;  // we always decompress the block before handing it over to the caller
 
 		FBufferView Payload = Response.Buffer.View(Reader.Tell(), BlobHeader.PayloadSize);
 		FBuffer		DecompressedBuffer;
@@ -150,7 +153,7 @@ FDownloadResult FHordeProtocolImpl::Download(const TArrayView<FNeedBlock> NeedBl
 			return FDownloadError(EDownloadRetryMode::Abort);
 		}
 
-		FHash128	 BlockHash128 = BlockHash.ToHash128();
+		FHash128 BlockHash128 = BlockHash.ToHash128();
 
 		CompletionCallback(DownloadedBlock, BlockHash128);
 
@@ -195,7 +198,8 @@ FHordeProtocolImpl::DownloadManifest(std::string_view ManifestName)
 	return DecodeHordeManifestJson((const char*)Response.Buffer.Data(), ManifestName);
 }
 
-bool RequestPathLooksLikeHordeArtifact(std::string_view RequestPath)
+bool
+RequestPathLooksLikeHordeArtifact(std::string_view RequestPath)
 {
 	static const std::regex Pattern("^api\\/v\\d+\\/artifacts\\/[a-fA-F0-9]+$");
 	return std::regex_match(RequestPath.begin(), RequestPath.end(), Pattern);
@@ -227,17 +231,17 @@ FHordeProtocolImpl::QueryHello(FHttpConnection& HttpConnection)
 		return AppError(std::string("JSON parse error while connecting to Horde server: ") + JsonErrorString);
 	}
 
-	if (auto& Field = JsonObject["serverUrl"]; Field.is_string())
+	if (const Json& Field = JsonObject["serverUrl"]; Field.is_string())
 	{
 		Result.AuthServerUri = Field.string_value();
 	}
 
-	if (auto& Field = JsonObject["clientId"]; Field.is_string())
+	if (const Json& Field = JsonObject["clientId"]; Field.is_string())
 	{
 		Result.AuthClientId = Field.string_value();
 	}
 
-	if (auto& Field = JsonObject["localRedirectUrls"]; Field.is_array())
+	if (const Json& Field = JsonObject["localRedirectUrls"]; Field.is_array())
 	{
 		if (Field.array_items().size() && Field.array_items()[0].is_string())
 		{
@@ -249,6 +253,224 @@ FHordeProtocolImpl::QueryHello(FHttpConnection& HttpConnection)
 	return ResultOk(Result);
 }
 
+TResult<FHordeVirtualPath>
+FHordeVirtualPath::FromString(std::string_view Str)
+{
+	FHordeVirtualPath Result;
+	size_t			  ChangeSeparatorPos	 = Str.find_first_of('@');
+	size_t			  ArtifactIdSeparatorPos = Str.find_first_of('#');
+
+	std::string_view ArtifactPrefix;
+	std::string_view StreamId;
+	std::string_view Change;
+	std::string_view ArtifactId;
+
+	if (ArtifactIdSeparatorPos != std::string::npos)
+	{
+		ArtifactId = Str.substr(ArtifactIdSeparatorPos + 1);
+	}
+
+	if (ChangeSeparatorPos != std::string::npos)
+	{
+		Change = Str.substr(ChangeSeparatorPos + 1);
+		Str	   = Str.substr(0, ChangeSeparatorPos);
+	}
+
+	if (!ArtifactId.empty())
+	{
+		Result.ArtifactId = ArtifactId;
+	}
+
+	size_t StreamSeparatorPos = Str.find_first_of("/\\");
+
+	if (StreamSeparatorPos != std::string::npos)
+	{
+		ArtifactPrefix = Str.substr(0, StreamSeparatorPos);
+		StreamId	   = Str.substr(StreamSeparatorPos + 1);
+	}
+	else
+	{
+		ArtifactPrefix = Str;
+	}
+
+	Result.ArtifactPrefix = ArtifactPrefix;
+
+	if (!StreamId.empty())
+	{
+		Result.StreamId = StreamId;
+	}
+
+	if (!Change.empty())
+	{
+		std::string ChangeStr = std::string(Change);
+
+		long long unsigned ParsedChange = 0;
+		sscanf(ChangeStr.c_str(), "%llu", &ParsedChange);
+
+		Result.Change = ParsedChange;
+	}
+
+	return ResultOk(Result);
+}
+
+TResult<ProxyQuery::FDirectoryListing>
+FHordeProtocolImpl::QueryListDirectory(FHttpConnection& Connection, const FAuthDesc* AuthDesc, const std::string& Path)
+{
+	using ProxyQuery::FDirectoryListing;
+	using ProxyQuery::FDirectoryListingEntry;
+
+	FDirectoryListing Result;
+
+	TResult<FHordeVirtualPath> VirtualPath = FHordeVirtualPath::FromString(Path);
+	if (VirtualPath.IsError())
+	{
+		return MoveError<ProxyQuery::FDirectoryListing>(VirtualPath);
+	}
+
+	// If we have the artifact ID, simply add a virtual .unsync directory to indicate that a manifest is available.
+	// TODO: It's also possible to list the actual artifact contents from this point, if needed.
+	if (VirtualPath->ArtifactId)
+	{
+		FDirectoryListingEntry DummyManifestEntry;
+		DummyManifestEntry.Name = ".unsync";
+		DummyManifestEntry.bDirectory = true;
+		Result.Entries.push_back(DummyManifestEntry);
+		return ResultOk(Result);
+	}
+
+	uint32 MaxResults = 1000000;
+
+	std::string RequestUrl = fmt::format("/api/v2/artifacts?type=packaged-build&maxResults={}", MaxResults);
+	if (VirtualPath->Change)
+	{
+		uint64 ChangeNumber = VirtualPath->Change.value();
+		RequestUrl += fmt::format("&minChange={0}&maxChange={0}", ChangeNumber);
+	}
+
+	if (VirtualPath->StreamId)
+	{
+		const std::string& StreamId = VirtualPath->StreamId.value();
+		RequestUrl += fmt::format("&streamId={}", StreamId);
+	}
+
+	std::string BearerToken;
+
+	TResult<FAuthToken> AuthTokenResult = Authenticate(*AuthDesc);
+	if (const FAuthToken* AuthToken = AuthTokenResult.TryData())
+	{
+		BearerToken = AuthToken->Access;
+	}
+
+	FHttpRequest Request;
+	Request.Url			= RequestUrl;
+	Request.BearerToken = BearerToken;
+
+	FHttpResponse Response = HttpRequest(Connection, Request);
+
+	if (!Response.Success())
+	{
+		return HttpError(Response.Code);
+	}
+
+	Response.Buffer.PushBack(0);
+
+	using namespace json11;
+	std::string JsonErrorString;
+	Json		JsonObject = Json::parse((const char*)Response.Buffer.Data(), JsonErrorString);
+
+	if (!JsonErrorString.empty())
+	{
+		return AppError(std::string("JSON parse error while listing Horde artifacts: ") + JsonErrorString);
+	}
+
+	Json JsonArtifacts = JsonObject["artifacts"];
+	if (!JsonArtifacts.is_array())
+	{
+		return AppError(L"Horde artifact listing is expected to have 'artifacts' array field");
+	}
+
+	std::vector<FHordeArtifactEntry> HordeArtifacts;
+
+	for (const Json& It : JsonArtifacts.array_items())
+	{
+		FHordeArtifactEntry Entry;
+
+		Entry.Name = It["name"].string_value();
+
+		if (!Entry.Name.starts_with(VirtualPath->ArtifactPrefix))
+		{
+			continue;
+		}
+
+		Entry.Change	  = uint64(It["change"].number_value());
+		Entry.Id		  = It["id"].string_value();
+		Entry.Description = It["description"].string_value();
+		Entry.StreamId	  = It["streamId"].string_value();
+
+		for (const Json& KeyJson : It["keys"].array_items())
+		{
+			const std::string& Key = KeyJson.string_value();
+			if (!Key.empty())
+			{
+				Entry.Keys.push_back(Key);
+			}
+		}
+
+		for (const Json& MetadataJson : It["metadata"].array_items())
+		{
+			const std::string& Metadata = MetadataJson.string_value();
+			if (!Metadata.empty())
+			{
+				Entry.Metadata.push_back(Metadata);
+			}
+		}
+
+		HordeArtifacts.push_back(Entry);
+	}
+
+	std::sort(HordeArtifacts.begin(),
+			  HordeArtifacts.end(),
+			  [](const FHordeArtifactEntry& A, const FHordeArtifactEntry& B) { return A.Change > B.Change; });
+
+	// Present a virtual two-level hierarchy: @change/package
+
+	if (VirtualPath->Change)
+	{
+		for (const FHordeArtifactEntry& HordeArtifactEntry : HordeArtifacts)
+		{
+			FDirectoryListingEntry DirectoryEntry;
+
+			DirectoryEntry.bDirectory = true;
+			DirectoryEntry.Name		  = fmt::format("{}#{}", HordeArtifactEntry.Name, HordeArtifactEntry.Id);
+
+			Result.Entries.push_back(DirectoryEntry);
+		}
+	}
+	else
+	{
+		// Output virtual directories for unique CLs
+		uint64			  LastChangeNumber = 0;
+		for (const FHordeArtifactEntry& HordeArtifactEntry : HordeArtifacts)
+		{
+			if (LastChangeNumber == HordeArtifactEntry.Change)
+			{
+				continue;
+			}
+
+			FDirectoryListingEntry DirectoryEntry;
+
+			DirectoryEntry.bDirectory = true;
+			DirectoryEntry.Name		  = fmt::format("{}@{}", HordeArtifactEntry.StreamId, HordeArtifactEntry.Change);
+
+			Result.Entries.push_back(DirectoryEntry);
+
+			LastChangeNumber = HordeArtifactEntry.Change;
+		}
+	}
+
+	return ResultOk(Result);
+}
+
 TResult<FDirectoryManifest> DecodeHordeManifestJson(const char* JsonString, std::string_view ArtifactRoot)
 {
 	using namespace json11;
@@ -270,7 +492,7 @@ TResult<FDirectoryManifest> DecodeHordeManifestJson(const char* JsonString, std:
 		return AppError("Manifest JSON is expected to have a 'type' string field with 'unsync_manifest' value");
 	}
 
-	if (auto& Field = JsonObject["hash_strong"]; Field.is_string())
+	if (const Json& Field = JsonObject["hash_strong"]; Field.is_string())
 	{
 		const std::string Value = StringToLower(Field.string_value());
 
@@ -296,7 +518,7 @@ TResult<FDirectoryManifest> DecodeHordeManifestJson(const char* JsonString, std:
 		}
 	}
 
-	if (auto& Field = JsonObject["hash_weak"]; Field.is_string())
+	if (const Json& Field = JsonObject["hash_weak"]; Field.is_string())
 	{
 		const std::string Value = StringToLower(Field.string_value());
 
@@ -314,7 +536,7 @@ TResult<FDirectoryManifest> DecodeHordeManifestJson(const char* JsonString, std:
 		}
 	}
 
-	if (auto& Field = JsonObject["chunking"]; Field.is_string())
+	if (const Json& Field = JsonObject["chunking"]; Field.is_string())
 	{
 		const std::string Value = StringToLower(Field.string_value());
 
@@ -341,9 +563,9 @@ TResult<FDirectoryManifest> DecodeHordeManifestJson(const char* JsonString, std:
 	const uint64 CurrentWindowsFileTime = ToWindowsFileTime(FileTimeNow);
 	uint32		 NumInvalidTimestamps	= 0;
 
-	if (auto& FiledField = JsonObject["files"]; FiledField.is_array())
+	if (const Json& FiledField = JsonObject["files"]; FiledField.is_array())
 	{
-		for (auto& FileObject : FiledField.array_items())
+		for (const Json& FileObject : FiledField.array_items())
 		{
 			std::string FileNameUtf8 = FileObject["name"].string_value();
 			std::wstring FileName	  = ConvertUtf8ToWide(FileNameUtf8);
@@ -360,12 +582,12 @@ TResult<FDirectoryManifest> DecodeHordeManifestJson(const char* JsonString, std:
 			FileManifest.BlockSize = DefaultBlockSize;
 			FileManifest.CurrentPath = FileName;
 
-			if (auto& Field = FileObject["size"]; Field.is_number())
+			if (const Json& Field = FileObject["size"]; Field.is_number())
 			{
 				FileManifest.Size = uint64(Field.number_value());
 			}
 
-			if (auto& Field = FileObject["mtime"]; Field.is_number())
+			if (const Json& Field = FileObject["mtime"]; Field.is_number())
 			{
 				FileManifest.Mtime = uint64(Field.number_value());
 			}
@@ -376,14 +598,14 @@ TResult<FDirectoryManifest> DecodeHordeManifestJson(const char* JsonString, std:
 				++NumInvalidTimestamps;
 			}
 
-			if (auto& Field = FileObject["read_only"]; Field.is_bool())
+			if (const Json& Field = FileObject["read_only"]; Field.is_bool())
 			{
 				FileManifest.bReadOnly = Field.bool_value();
 			}
 
-			if (auto& BlocksField = FileObject["blocks"]; BlocksField.is_array())
+			if (const Json& BlocksField = FileObject["blocks"]; BlocksField.is_array())
 			{
-				for (auto& BlockObject : BlocksField.array_items())
+				for (const Json& BlockObject : BlocksField.array_items())
 				{
 					FGenericBlock Block;
 					Block.Offset = uint64(BlockObject["offset"].number_value());
diff --git a/Engine/Source/Programs/Unsync/Private/UnsyncHorde.h b/Engine/Source/Programs/Unsync/Private/UnsyncHorde.h
--- a/Engine/Source/Programs/Unsync/Private/UnsyncHorde.h
+++ b/Engine/Source/Programs/Unsync/Private/UnsyncHorde.h
@@ -3,10 +3,13 @@
 #include "UnsyncCommon.h"
 #include "UnsyncProxy.h"
 
+#include <optional>
+
 namespace unsync {
 
 class FBlockRequestMap;
 struct FRemoteDesc;
+struct FAuthDesc;
 
 struct FHordeProtocolImpl : FRemoteProtocolBase
 {
@@ -18,13 +21,38 @@ struct FHordeProtocolImpl : FRemoteProtocolBase
 	virtual FDownloadResult	 Download(const TArrayView<FNeedBlock> NeedBlocks, const FBlockDownloadCallback& CompletionCallback);
 	virtual TResult<FDirectoryManifest> DownloadManifest(std::string_view ManifestName);
 
-
 	static TResult<ProxyQuery::FHelloResponse> QueryHello(FHttpConnection& HttpConnection);
+	static TResult<ProxyQuery::FDirectoryListing> QueryListDirectory(FHttpConnection&	Connection,
+																	 const FAuthDesc*	AuthDesc,
+																	 const std::string& Path);
 
 	FProxyPool& ProxyPool;
 	bool bValid = true;
 };
 
+
+struct FHordeVirtualPath
+{
+	std::string ArtifactPrefix;
+
+	std::optional<std::string> StreamId;
+	std::optional<std::string> ArtifactId;
+	std::optional<uint64>	   Change;
+
+	static TResult<FHordeVirtualPath> FromString(std::string_view Str);
+};
+
+struct FHordeArtifactEntry
+{
+	uint64					 Change = 0;
+	std::string				 Id;
+	std::string				 Name;
+	std::string				 Description;
+	std::string				 StreamId;
+	std::vector<std::string> Keys;
+	std::vector<std::string> Metadata;
+};
+
 bool RequestPathLooksLikeHordeArtifact(std::string_view RequestPath);
 
 TResult<FDirectoryManifest> DecodeHordeManifestJson(const char* JsonString, std::string_view ArtifactRoot);
diff --git a/Engine/Source/Programs/Unsync/Private/UnsyncMain.cpp b/Engine/Source/Programs/Unsync/Private/UnsyncMain.cpp
--- a/Engine/Source/Programs/Unsync/Private/UnsyncMain.cpp
+++ b/Engine/Source/Programs/Unsync/Private/UnsyncMain.cpp
@@ -22,6 +22,7 @@
 #include "UnsyncVersion.h"
 #include "UnsyncSource.h"
 #include "UnsyncFilter.h"
+#include "UnsyncHorde.h"
 
 UNSYNC_THIRD_PARTY_INCLUDES_START
 #if UNSYNC_PLATFORM_WINDOWS
@@ -662,6 +663,31 @@ InnerMain(int Argc, char** Argv)
 		}
 	}
 
+	// Derive artifact request path when syncing from Horde, if it wasn't specified via URL source syntax
+	if (RemoteDesc.Protocol == EProtocolFlavor::Horde && Cli.got_subcommand(SubSync))
+	{
+		bFilesystemSource = false;
+		if (RemoteDesc.RequestPath.empty())
+		{
+			TResult<FHordeVirtualPath> VirtualPath = FHordeVirtualPath::FromString(SourceFilenameUtf8);
+			if (VirtualPath.IsError())
+			{
+				UNSYNC_ERROR(L"Could not parse sync source path");
+				LogError(VirtualPath.GetError());
+				return 1;
+			}
+
+			if (!VirtualPath->ArtifactId)
+			{
+				UNSYNC_ERROR(L"Could not parse sync source path. Artifact ID is expected, i.e. '#123456abcdef'.");
+				return 1;
+			}
+
+			SourceFilenameUtf8 = "api/v2/artifacts/" + VirtualPath->ArtifactId.value();
+			RemoteDesc.RequestPath = SourceFilenameUtf8;
+		}
+	}
+
 	if (bNoCompression)
 	{
 		UNSYNC_VERBOSE(L"Uncompressed data transfer is preferred");
@@ -985,7 +1011,7 @@ InnerMain(int Argc, char** Argv)
 
 				// Note: since tokens can expire during a long operation,
 				// we can only save the auth descriptor and re-authenticate later if necessary
-				TResult<FAuthToken> AuthTokenResult = Authenticate(AuthDesc, 5 * 60);
+				TResult<FAuthToken> AuthTokenResult = Authenticate(AuthDesc);
 
 				if (AuthTokenResult.IsError())
 				{
diff --git a/Engine/Source/Programs/Unsync/Private/UnsyncProxy.cpp b/Engine/Source/Programs/Unsync/Private/UnsyncProxy.cpp
--- a/Engine/Source/Programs/Unsync/Private/UnsyncProxy.cpp
+++ b/Engine/Source/Programs/Unsync/Private/UnsyncProxy.cpp
@@ -46,8 +46,8 @@ struct FUnsyncProtocolImpl : FUnsyncBaseProtocolImpl
 						const FAuthDesc*			   InAuthDesc,
 						const FBlockRequestMap*		   InRequestMap);
 	virtual ~FUnsyncProtocolImpl() override;
-	virtual bool			 IsValid() const override;
-	virtual FDownloadResult	 Download(const TArrayView<FNeedBlock> NeedBlocks, const FBlockDownloadCallback& CompletionCallback) override;
+	virtual bool			IsValid() const override;
+	virtual FDownloadResult Download(const TArrayView<FNeedBlock> NeedBlocks, const FBlockDownloadCallback& CompletionCallback) override;
 
 	virtual void Invalidate() override;
 
@@ -57,22 +57,26 @@ struct FUnsyncProtocolImpl : FUnsyncBaseProtocolImpl
 	std::unique_ptr<FSocketBase> SocketHandle;
 
 	static void SendTelemetryEvent(const FRemoteDesc& RemoteDesc, const FTelemetryEventSyncComplete& Event);
-	static TResult<ProxyQuery::FHelloResponse> QueryHello(FHttpConnection& HttpConnection, const FAuthDesc* OptAuthDesc=nullptr);
+
+	static TResult<ProxyQuery::FHelloResponse>	  QueryHello(FHttpConnection& HttpConnection, const FAuthDesc* OptAuthDesc = nullptr);
+	static TResult<ProxyQuery::FDirectoryListing> QueryListDirectory(FHttpConnection&	Connection,
+																	 const FAuthDesc*	AuthDesc,
+																	 const std::string& Path);
 };
 
 struct FUnsyncHttpProtocolImpl : FUnsyncBaseProtocolImpl
 {
-	FUnsyncHttpProtocolImpl(const FRemoteDesc& InRemoteDesc,
+	FUnsyncHttpProtocolImpl(const FRemoteDesc&			   InRemoteDesc,
 							const FRemoteProtocolFeatures& InFeatures,
-							const FBlockRequestMap* InRequestMap,
-							FProxyPool&				   InProxyPool)
+							const FBlockRequestMap*		   InRequestMap,
+							FProxyPool&					   InProxyPool)
 	: FUnsyncBaseProtocolImpl(InRemoteDesc, InRequestMap, InFeatures)
 	, ProxyPool(InProxyPool)
 	{
 	}
 
-	virtual bool			IsValid() const override { return bValid; }
-	virtual void			Invalidate() override { bValid = false; }
+	virtual bool IsValid() const override { return bValid; }
+	virtual void Invalidate() override { bValid = false; }
 
 	virtual FDownloadResult Download(const TArrayView<FNeedBlock> NeedBlocks, const FBlockDownloadCallback& CompletionCallback) override;
 
@@ -126,8 +130,8 @@ FUnsyncProtocolImpl::FUnsyncProtocolImpl(const FRemoteDesc&				RemoteDesc,
 {
 	if (RemoteDesc.TlsRequirement != ETlsRequirement::None)
 	{
-		FTlsClientSettings TlsSettings = RemoteDesc.GetTlsClientSettings();
-		FSocketHandle RawSocketHandle = SocketConnectTcp(RemoteDesc.Host.Address.c_str(), RemoteDesc.Host.Port);
+		FTlsClientSettings TlsSettings	   = RemoteDesc.GetTlsClientSettings();
+		FSocketHandle	   RawSocketHandle = SocketConnectTcp(RemoteDesc.Host.Address.c_str(), RemoteDesc.Host.Port);
 		SocketSetRecvTimeout(RawSocketHandle, RemoteDesc.RecvTimeoutSeconds);
 
 		if (RawSocketHandle)
@@ -188,7 +192,7 @@ FUnsyncProtocolImpl::FUnsyncProtocolImpl(const FRemoteDesc&				RemoteDesc,
 		{
 			bool bOk = IsValid();
 
-			TResult<FAuthToken> AuthTokenResult = Authenticate(*InAuthDesc, 15 * 60);
+			TResult<FAuthToken> AuthTokenResult = Authenticate(*InAuthDesc);
 
 			if (AuthTokenResult.IsOk())
 			{
@@ -497,14 +501,14 @@ FormatBlockRequestJson(const FBlockRequestMap& RequestMap, const TArrayView<FNee
 
 	std::string Output;
 
-	Output += "{ ";  // main object
+	Output += "{ ";	 // main object
 
 	FormatJsonKeyValueStr(Output, "hash_strong", StrongHashAlgorithm, ",\n");
 
 	Output += "\"files\": [\n";
 
-	static const FHash128 InvalidHash = {};
-	FHash128	   FilenameHash = InvalidHash;
+	static const FHash128 InvalidHash  = {};
+	FHash128			  FilenameHash = InvalidHash;
 
 	uint32 BlockIndex = 0;
 	for (const FBlockRequestAndHash& Request : Requests)
@@ -513,7 +517,7 @@ FormatBlockRequestJson(const FBlockRequestMap& RequestMap, const TArrayView<FNee
 		{
 			if (FilenameHash != InvalidHash)
 			{
-				Output += "]},\n";  // close blocks arary and file object
+				Output += "]},\n";	// close blocks arary and file object
 			}
 
 			const std::string* FilenameUtf8 = RequestMap.FindFile(Request.FilenameMd5);
@@ -526,7 +530,7 @@ FormatBlockRequestJson(const FBlockRequestMap& RequestMap, const TArrayView<FNee
 			FormatJsonKeyValueStr(Output, "name", EscapedFilenameUtf8, ", ");
 			Output += "\"blocks\": [\n";
 
-			BlockIndex = 0;
+			BlockIndex	 = 0;
 			FilenameHash = Request.FilenameMd5;
 		}
 
@@ -547,7 +551,7 @@ FormatBlockRequestJson(const FBlockRequestMap& RequestMap, const TArrayView<FNee
 
 	if (FilenameHash != InvalidHash)
 	{
-		Output += "]}\n";	// close blocks arary and file object
+		Output += "]}\n";  // close blocks arary and file object
 	}
 
 	Output += "]\n";  // files array
@@ -663,7 +667,7 @@ FUnsyncProtocolImpl::QueryHello(FHttpConnection& HttpConnection, const FAuthDesc
 	std::string BearerToken;
 	if (OptAuthDesc)
 	{
-		TResult<FAuthToken> AuthTokenResult = Authenticate(*OptAuthDesc, 15 * 60);
+		TResult<FAuthToken> AuthTokenResult = Authenticate(*OptAuthDesc);
 		if (AuthTokenResult.IsOk())
 		{
 			BearerToken = std::move(AuthTokenResult.GetData().Access);
@@ -861,14 +865,33 @@ ProxyQuery::FDirectoryListing::ToJson() const
 }
 
 TResult<ProxyQuery::FDirectoryListing>
-ProxyQuery::ListDirectory(FHttpConnection& Connection, const FAuthDesc* AuthDesc, const std::string& Path)
+ProxyQuery::ListDirectory(EProtocolFlavor Protocol, FHttpConnection& Connection, const FAuthDesc* AuthDesc, const std::string& Path)
+{
+	if (Protocol == EProtocolFlavor::Horde)
+	{
+		return FHordeProtocolImpl::QueryListDirectory(Connection, AuthDesc, Path);
+	}
+	else if (Protocol == EProtocolFlavor::Unsync)
+	{
+		return FUnsyncProtocolImpl::QueryListDirectory(Connection, AuthDesc, Path);
+	}
+	else
+	{
+		return AppError("Protocol does not support server directory listing");
+	}
+}
+
+TResult<ProxyQuery::FDirectoryListing>
+FUnsyncProtocolImpl::QueryListDirectory(FHttpConnection& Connection, const FAuthDesc* AuthDesc, const std::string& Path)
 {
+	using ProxyQuery::FDirectoryListing;
+
 	std::string Url = fmt::format("/api/v1/list?{}", Path);
 
 	std::string BearerToken;
 	if (AuthDesc)
 	{
-		TResult<FAuthToken> AuthToken = Authenticate(*AuthDesc, 5 * 60);
+		TResult<FAuthToken> AuthToken = Authenticate(*AuthDesc);
 		if (!AuthToken.IsOk())
 		{
 			return MoveError<FDirectoryListing>(AuthToken);
@@ -916,7 +939,7 @@ ProxyQuery::DownloadFile(FHttpConnection&					 InConnection,
 		std::string BearerToken;
 		if (AuthDesc)
 		{
-			TResult<FAuthToken> AuthToken = Authenticate(*AuthDesc, 5 * 60);
+			TResult<FAuthToken> AuthToken = Authenticate(*AuthDesc);
 			if (!AuthToken.IsOk())
 			{
 				return std::move(AuthToken.GetError());
@@ -967,7 +990,7 @@ ProxyQuery::DownloadFile(FHttpConnection&					 InConnection,
 	std::string BearerToken;
 	if (AuthDesc)
 	{
-		TResult<FAuthToken> AuthToken = Authenticate(*AuthDesc, 5 * 60);
+		TResult<FAuthToken> AuthToken = Authenticate(*AuthDesc);
 		if (!AuthToken.IsOk())
 		{
 			Error.Set(std::move(AuthToken.GetError()));
@@ -1373,7 +1396,7 @@ FProxyPool::GetAccessToken()
 
 	if (AuthDesc)
 	{
-		TResult<FAuthToken> AuthTokenResult = Authenticate(*AuthDesc, 120);
+		TResult<FAuthToken> AuthTokenResult = Authenticate(*AuthDesc);
 		if (FAuthToken* AuthToken = AuthTokenResult.TryData())
 		{
 			std::swap(Result, AuthToken->Access);
@@ -1488,7 +1511,7 @@ FRemoteFileSystem::ListDirectory(const std::string_view RelativePath)
 		FullPath.append(RelativePath);
 	}
 	ConvertDirectorySeparatorsToUnix(FullPath);
-	return ProxyQuery::ListDirectory(HttpConnection, ProxyPool.AuthDesc, FullPath);
+	return ProxyQuery::ListDirectory(ProxyPool.RemoteDesc.Protocol, HttpConnection, ProxyPool.AuthDesc, FullPath);
 }
 
 TResult<FBuffer>
diff --git a/Engine/Source/Programs/Unsync/Private/UnsyncProxy.h b/Engine/Source/Programs/Unsync/Private/UnsyncProxy.h
--- a/Engine/Source/Programs/Unsync/Private/UnsyncProxy.h
+++ b/Engine/Source/Programs/Unsync/Private/UnsyncProxy.h
@@ -265,7 +265,11 @@ struct FDirectoryListing
 	std::string						  ToJson() const;
 };
 
-TResult<FDirectoryListing> ListDirectory(FHttpConnection& Connection, const FAuthDesc* AuthDesc, const std::string& Path);
+TResult<FDirectoryListing> ListDirectory(EProtocolFlavor	Protocol,
+										 FHttpConnection&	Connection,
+										 const FAuthDesc*	AuthDesc,
+										 const std::string& Path);
+
 TResult<FBuffer>		   DownloadFile(FHttpConnection& Connection, const FAuthDesc* AuthDesc, const std::string& Path);
 
 using FDownloadOutputCallback = std::function<FIOWriter&(uint64 Size)>;
diff --git a/Engine/Source/Programs/Unsync/Private/UnsyncRemote.cpp b/Engine/Source/Programs/Unsync/Private/UnsyncRemote.cpp
--- a/Engine/Source/Programs/Unsync/Private/UnsyncRemote.cpp
+++ b/Engine/Source/Programs/Unsync/Private/UnsyncRemote.cpp
@@ -6,23 +6,6 @@
 
 namespace unsync {
 
-EProtocolFlavor
-protocol_flavor_from_string(const char* Str)
-{
-	if (!strcmp(Str, "jupiter"))
-	{
-		return EProtocolFlavor::Jupiter;
-	}
-	if (!strcmp(Str, "unsync"))
-	{
-		return EProtocolFlavor::Unsync;
-	}
-	else
-	{
-		return EProtocolFlavor::Unknown;
-	}
-}
-
 const char*
 ToString(EProtocolFlavor Protocol)
 {
