unsync - Add Horde artifact listing syntax backwards-compatible with filesystem based paths

* Specify artifact filters using `KeyA=ValueA;KeyB=ValueB` syntax
* Specify artifact listing formating and directory structure using `{FieldA}{dir}{FieldB}` syntax, where {FieldA/B} are substituted with artifact metadata and {dir} creates a virtual sub-directory
* Missing fields are substituted with `[unknown]`

diff --git a/Engine/Source/Programs/Unsync/Private/UnsyncCmdQuery.cpp b/Engine/Source/Programs/Unsync/Private/UnsyncCmdQuery.cpp
--- a/Engine/Source/Programs/Unsync/Private/UnsyncCmdQuery.cpp
+++ b/Engine/Source/Programs/Unsync/Private/UnsyncCmdQuery.cpp
@@ -274,10 +274,11 @@ CmdQuerySearch(const FCmdQueryOptions& Options)
 
 	struct FTaskContext
 	{
-		std::mutex				  Mutex;
-		std::vector<FResultEntry> FoundEntries;
-		bool					  bParentThreadVerbose = false;
-		int32					  ParentThreadIndent   = 0;
+		std::mutex						Mutex;
+		std::vector<FResultEntry>		FoundEntries;
+		std::unordered_set<std::string> VisitedDirectories;
+		bool							bParentThreadVerbose = false;
+		int32							ParentThreadIndent	 = 0;
 	};
 
 	FTaskGroup	 Tasks = GScheduler->CreateTaskGroup();
@@ -286,37 +287,59 @@ CmdQuerySearch(const FCmdQueryOptions& Options)
 	Context.bParentThreadVerbose = GLogVerbose;
 	Context.ParentThreadIndent	 = GLogIndent;
 
-	std::function<void(std::string, int32)> ExploreDirectory =
-		[&Context, &AuthDesc, &ConnectionPool, &ExploreDirectory, &SubdirPatterns, &Tasks, &Options](std::string Path, int32 CurrentDepth)
+	std::function<void(std::string, int32, const FDirectoryListing*)> ExploreDirectory =
+		[&Context, &AuthDesc, &ConnectionPool, &ExploreDirectory, &SubdirPatterns, &Tasks, &Options](
+			std::string					   Path,
+			int32						   CurrentDepth,
+			const FDirectoryListing* DirectoryListingPtr)
 	{
 		FLogVerbosityScope VerboseScope(Context.bParentThreadVerbose);
 		FLogIndentScope	   IndentScope(Context.ParentThreadIndent, true);
 
-		UNSYNC_VERBOSE2(L"Listing '%hs'", Path.c_str());
+		FDirectoryListing RemoteDirectoryListing;
 
-		GScheduler->NetworkSemaphore.Acquire(false);
-		std::unique_ptr<FHttpConnection> Connection = ConnectionPool.Acquire();
+		if (!DirectoryListingPtr)
+		{
+			UNSYNC_VERBOSE2(L"Listing '%hs'", Path.c_str());
 
-		TResult<FDirectoryListing> DirectoryListingResult =
-			ProxyQuery::ListDirectory(Options.Remote.Protocol, *Connection, &AuthDesc, Path);
+			GScheduler->NetworkSemaphore.Acquire(false);
+			std::unique_ptr<FHttpConnection> Connection = ConnectionPool.Acquire();
 
-		ConnectionPool.Release(std::move(Connection));
-		GScheduler->NetworkSemaphore.Release();
+			TResult<FDirectoryListing> DirectoryListingResult =
+				ProxyQuery::ListDirectory(Options.Remote.Protocol, *Connection, &AuthDesc, Path);
 
-		if (DirectoryListingResult.IsError())
-		{
-			LogError(DirectoryListingResult.GetError());
-			return;
+			ConnectionPool.Release(std::move(Connection));
+			GScheduler->NetworkSemaphore.Release();
+
+			if (DirectoryListingResult.IsError())
+			{
+				LogError(DirectoryListingResult.GetError());
+				return;
+			}
+
+			std::swap(RemoteDirectoryListing, DirectoryListingResult.GetData());
+			DirectoryListingPtr = &RemoteDirectoryListing;
 		}
 
-		const FDirectoryListing& DirectoryListing = DirectoryListingResult.GetData();
+		const FDirectoryListing& DirectoryListing = *DirectoryListingPtr;
 
 		for (const FDirectoryListingEntry& DirEntry : DirectoryListing.Entries)
 		{
+			std::string DirEntryName = DirEntry.Name;
+
+			// Only include one subdirectory level
+			const size_t SeparatorPos = DirEntryName.find(PATH_SEPARATOR);
+			if (SeparatorPos != std::string::npos)
+			{
+				DirEntryName = DirEntryName.substr(0, PATH_SEPARATOR);
+			}
+
 			FEntry NextEntry;
-			NextEntry.Path	= Path + "\\" + DirEntry.Name;
+			NextEntry.Path	= Path + PATH_SEPARATOR + DirEntryName;
 			NextEntry.Depth = CurrentDepth + 1;
 
+			// Include only leaf directory entries in the final output
+			if (SeparatorPos == std::string::npos)
 			{
 				std::lock_guard<std::mutex> LockGuard(Context.Mutex);
 
@@ -333,15 +356,46 @@ CmdQuerySearch(const FCmdQueryOptions& Options)
 				continue;
 			}
 
-			if (DirEntry.bDirectory && std::regex_match(DirEntry.Name, SubdirPatterns[CurrentDepth], std::regex_constants::match_any))
+			if (DirEntry.bDirectory && std::regex_match(DirEntryName, SubdirPatterns[CurrentDepth], std::regex_constants::match_any))
 			{
-				UNSYNC_VERBOSE2(L"Matched: '%hs'", DirEntry.Name.c_str());
-				Tasks.run([ExploreDirectory, NextEntry]() { ExploreDirectory(NextEntry.Path, NextEntry.Depth); });
+				UNSYNC_VERBOSE2(L"Matched: '%hs'", DirEntryName.c_str());
+
+				// Directory listing may already include some child sub-directories, so we can skip some network requests
+				FDirectoryListing SubDirectoryListing;
+				std::string		  RequiredPrefix = DirEntryName + PATH_SEPARATOR;
+
+				for (const FDirectoryListingEntry& DirEntry2 : DirectoryListing.Entries)
+				{
+					if (DirEntry2.Name.starts_with(RequiredPrefix))
+					{
+						FDirectoryListingEntry SubDirEntry = DirEntry2;
+						SubDirEntry.Name				   = SubDirEntry.Name.substr(RequiredPrefix.length());
+						SubDirectoryListing.Entries.push_back(SubDirEntry);
+					}
+				}
+
+				{
+					// Only visit each sub-directory once
+					std::lock_guard<std::mutex> LockGuard(Context.Mutex);
+					if (!Context.VisitedDirectories.insert(NextEntry.Path).second)
+					{
+						continue;
+					}
+				}
+
+				if (SubDirectoryListing.Entries.empty())
+				{
+					Tasks.run([ExploreDirectory, NextEntry]() { ExploreDirectory(NextEntry.Path, NextEntry.Depth, nullptr); });
+				}
+				else
+				{
+					ExploreDirectory(NextEntry.Path, NextEntry.Depth, &SubDirectoryListing);
+				}
 			}
 		}
 	};
 
-	ExploreDirectory(RootPath, 0);
+	ExploreDirectory(RootPath, 0, nullptr);
 	Tasks.wait();
 
 	std::vector<FResultEntry>& ResultEntries = Context.FoundEntries;
diff --git a/Engine/Source/Programs/Unsync/Private/UnsyncHorde.cpp b/Engine/Source/Programs/Unsync/Private/UnsyncHorde.cpp
--- a/Engine/Source/Programs/Unsync/Private/UnsyncHorde.cpp
+++ b/Engine/Source/Programs/Unsync/Private/UnsyncHorde.cpp
@@ -9,6 +9,7 @@
 #include <json11.hpp>
 #include <optional>
 #include <stdio.h>
+#include <span>
 
 namespace unsync {
 
@@ -253,61 +254,205 @@ FHordeProtocolImpl::QueryHello(FHttpConnection& HttpConnection)
 	return ResultOk(Result);
 }
 
-TResult<FHordeVirtualPath>
-FHordeVirtualPath::FromString(std::string_view Str)
+static std::optional<std::string_view>
+FindFirstByPrefix(std::span<const std::string> Strings, std::string_view Prefix)
 {
-	FHordeVirtualPath Result;
-	size_t			  ChangeSeparatorPos	 = Str.find_first_of('@');
-	size_t			  ArtifactIdSeparatorPos = Str.find_first_of('#');
+	for (const std::string& It : Strings)
+	{
+		if (StringStartsWith(It, Prefix, false))
+		{
+			return std::string_view(It);
+		}
+	}
 
-	std::string_view ArtifactPrefix;
-	std::string_view StreamId;
-	std::string_view Change;
-	std::string_view ArtifactId;
+	return {};
+}
 
-	if (ArtifactIdSeparatorPos != std::string::npos)
+std::optional<std::string_view>
+GetMetadataValueByPrefix(std::span<const std::string> MetadataEntries, std::string_view Prefix)
+{
+	std::optional<std::string_view> Entry = FindFirstByPrefix(MetadataEntries, Prefix);
+	if (Entry)
+	{
+		return Entry->substr(Prefix.length());
+	}
+	else
 	{
-		ArtifactId = Str.substr(ArtifactIdSeparatorPos + 1);
+		return {};
 	}
+}
 
-	if (ChangeSeparatorPos != std::string::npos)
+std::optional<std::string_view>
+TryStripPrefix(std::string_view String, std::string_view Prefix)
+{
+	if (StringStartsWith(String, Prefix, false))
+	{
+		return String.substr(Prefix.length());
+	}
+	else
 	{
-		Change = Str.substr(ChangeSeparatorPos + 1);
-		Str	   = Str.substr(0, ChangeSeparatorPos);
+		return {};
 	}
+}
+
+static std::string
+UrlEncode(std::string_view String)
+{
+	std::string Result;
+	Result.reserve(String.length() * 3);
 
-	if (!ArtifactId.empty())
+	for (char C : String)
 	{
-		Result.ArtifactId = ArtifactId;
+		if (IsAsciiAlphaNumericCharacter(C))
+		{
+			Result.push_back(C);
+		}
+		else
+		{
+			char HexChars[2] = {};
+			BytesToHexChars(HexChars, 2, reinterpret_cast<const uint8*>(&C), 1);
+			Result.push_back('%');
+			Result.push_back(HexChars[0]);
+			Result.push_back(HexChars[1]);
+		}
 	}
 
-	size_t StreamSeparatorPos = Str.find_first_of("/\\");
+	return Result;
+}
 
-	if (StreamSeparatorPos != std::string::npos)
+TResult<FHordeArtifactQuery>
+FHordeArtifactQuery::FromString(std::string_view Str)
+{
+	FHordeArtifactQuery Result;
+
+	std::vector<std::string_view> PathParts = SplitBy(Str, PATH_SEPARATOR);
+	std::string_view			  RootPath	= PathParts[0];
+
+	for (std::string_view PathPart : PathParts)
 	{
-		ArtifactPrefix = Str.substr(0, StreamSeparatorPos);
-		StreamId	   = Str.substr(StreamSeparatorPos + 1);
+		if (!Result.Parts.empty())
+		{
+			// TODO: derived the ID by parsing the request string using the format specification
+			size_t ArtifactIdPos = PathPart.find_last_of('#');
+			if (ArtifactIdPos != std::string::npos)
+			{
+				Result.Id = PathPart.substr(ArtifactIdPos + 1);
+			}
+		}
+
+		Result.Parts.push_back(std::string(PathPart));
 	}
-	else
+
+	std::vector<std::string_view> FilterParts = SplitBy(RootPath, ';');
+
+	for (std::string_view FilterPart : FilterParts)
 	{
-		ArtifactPrefix = Str;
+		std::optional<std::string_view> Value;
+		if (Value = TryStripPrefix(FilterPart, "type="); Value)
+		{
+			Result.Type = StringToLower(*Value);
+		}
+		else if (Value = TryStripPrefix(FilterPart, "key="); Value)
+		{
+			Result.Keys.push_back(UrlEncode(StringToLower(std::string(*Value))));
+		}
+		else if (Value = TryStripPrefix(FilterPart, "change="); Value)
+		{
+			Result.Change = *Value;
+		}
+		else if (Value = TryStripPrefix(FilterPart, "id="); Value)
+		{
+			Result.Id = *Value;
+		}
+		else
+		{
+			Result.Format = FilterPart;
+		}
 	}
 
-	Result.ArtifactPrefix = ArtifactPrefix;
+	return ResultOk(Result);
+}
 
-	if (!StreamId.empty())
+std::string
+GetEscapedFormattedNameField(std::string_view StreamName)
+{
+	std::string Result = std::string(StreamName);
+	for (char& C : Result)
 	{
-		Result.StreamId = StreamId;
+		if (C == '/' || C == '\\')
+		{
+			C = '+';
+		}
 	}
 
-	if (!Change.empty())
+	return Result;
+}
+
+TResult<std::string>
+FormatArtifactName(const FHordeArtifactEntry& Artifact, std::string_view Format, std::string_view PlaceholderValue = "[unknown]")
+{
+	std::string Result;
+
+	while (!Format.empty())
 	{
-		std::string ChangeStr = std::string(Change);
+		if (Format.front() == '{')
+		{
+			size_t ClosingPos = Format.find('}');
+			if (ClosingPos == std::string::npos)
+			{
+				return AppError(fmt::format("Expected to find '{}' character in the format template string: '{}'", '}', Format));
+			}
 
-		long long unsigned ParsedChange = 0;
-		sscanf(ChangeStr.c_str(), "%llu", &ParsedChange);
+			std::string_view VarName = Format.substr(1, ClosingPos - 1);
 
-		Result.Change = ParsedChange;
+			Format = Format.substr(ClosingPos + 1);
+
+			if (UncasedStringEquals(VarName, "dir"))
+			{
+				Result.push_back(PATH_SEPARATOR);
+				continue;
+			}
+
+			std::string Key = fmt::format("{}=", VarName);
+
+			// First try to get the value from metadata, then from artifact itself
+
+			std::optional<std::string_view> Value = GetMetadataValueByPrefix(Artifact.Metadata, Key);
+
+			if (!Value)
+			{
+				if (UncasedStringEquals(Key, "change="))
+				{
+					Value = Artifact.Change;
+				}
+				else if (UncasedStringEquals(Key, "id="))
+				{
+					Value = Artifact.Id;
+				}
+				else if (UncasedStringEquals(Key, "name="))
+				{
+					Value = Artifact.Name;
+				}
+				else if (UncasedStringEquals(Key, "streamid="))
+				{
+					Value = Artifact.StreamId;
+				}
+			}
+
+			if (Value)
+			{
+				Result.append(GetEscapedFormattedNameField(*Value));
+			}
+			else
+			{
+				Result.append(PlaceholderValue);
+			}
+		}
+		else
+		{
+			Result.push_back(Format.front());
+			Format = Format.substr(1);
+		}
 	}
 
 	return ResultOk(Result);
@@ -321,36 +466,41 @@ FHordeProtocolImpl::QueryListDirectory(FHttpConnection& Connection, const FAuthD
 
 	FDirectoryListing Result;
 
-	TResult<FHordeVirtualPath> VirtualPath = FHordeVirtualPath::FromString(Path);
-	if (VirtualPath.IsError())
+	TResult<FHordeArtifactQuery> QueryParseResult = FHordeArtifactQuery::FromString(Path);
+	if (QueryParseResult.IsError())
 	{
-		return MoveError<ProxyQuery::FDirectoryListing>(VirtualPath);
+		return MoveError<ProxyQuery::FDirectoryListing>(QueryParseResult);
 	}
 
-	// If we have the artifact ID, simply add a virtual .unsync directory to indicate that a manifest is available.
-	// TODO: It's also possible to list the actual artifact contents from this point, if needed.
-	if (VirtualPath->ArtifactId)
+	const FHordeArtifactQuery& Query = QueryParseResult.GetData();
+
+	std::string_view Root = Query.Parts[0];
+
+	std::string RequiredPathPrefix;
+	if (Root.length() != Path.length())
 	{
-		FDirectoryListingEntry DummyManifestEntry;
-		DummyManifestEntry.Name = ".unsync";
-		DummyManifestEntry.bDirectory = true;
-		Result.Entries.push_back(DummyManifestEntry);
-		return ResultOk(Result);
+		RequiredPathPrefix = Path.substr(Root.length() + 1);
+		if (RequiredPathPrefix.back() != (PATH_SEPARATOR))
+		{
+			RequiredPathPrefix.push_back(PATH_SEPARATOR);
+		}
 	}
 
-	uint32 MaxResults = 1000000;
+	std::string RequestUrl = fmt::format("/api/v2/artifacts?maxResults={}", Query.MaxResults);
 
-	std::string RequestUrl = fmt::format("/api/v2/artifacts?type=packaged-build&maxResults={}", MaxResults);
-	if (VirtualPath->Change)
+	if (!Query.Type.empty())
 	{
-		uint64 ChangeNumber = VirtualPath->Change.value();
-		RequestUrl += fmt::format("&minChange={0}&maxChange={0}", ChangeNumber);
+		RequestUrl += fmt::format("&type={}", Query.Type);
 	}
 
-	if (VirtualPath->StreamId)
+	if (!Query.Change.empty())
 	{
-		const std::string& StreamId = VirtualPath->StreamId.value();
-		RequestUrl += fmt::format("&streamId={}", StreamId);
+		RequestUrl += fmt::format("&minChange={0}&maxChange={0}", Query.Change);
+	}
+
+	for (const std::string& Key : Query.Keys)
+	{
+		RequestUrl += fmt::format("&key={}", Key);
 	}
 
 	std::string BearerToken;
@@ -397,12 +547,16 @@ FHordeProtocolImpl::QueryListDirectory(FHttpConnection& Connection, const FAuthD
 
 		Entry.Name = It["name"].string_value();
 
-		if (!Entry.Name.starts_with(VirtualPath->ArtifactPrefix))
+		const Json& ChangeJson = It["change"];
+		if (ChangeJson.is_number())
+		{
+			Entry.Change = fmt::format("{}", llu(ChangeJson.number_value()));
+		}
+		else if (ChangeJson.is_string())
 		{
-			continue;
+			Entry.Change = ChangeJson.string_value();
 		}
 
-		Entry.Change	  = uint64(It["change"].number_value());
 		Entry.Id		  = It["id"].string_value();
 		Entry.Description = It["description"].string_value();
 		Entry.StreamId	  = It["streamId"].string_value();
@@ -432,39 +586,80 @@ FHordeProtocolImpl::QueryListDirectory(FHttpConnection& Connection, const FAuthD
 			  HordeArtifacts.end(),
 			  [](const FHordeArtifactEntry& A, const FHordeArtifactEntry& B) { return A.Change > B.Change; });
 
-	// Present a virtual two-level hierarchy: @change/package
-
-	if (VirtualPath->Change)
+	if (!Query.Format.empty())
 	{
 		for (const FHordeArtifactEntry& HordeArtifactEntry : HordeArtifacts)
 		{
 			FDirectoryListingEntry DirectoryEntry;
 
-			DirectoryEntry.bDirectory = true;
-			DirectoryEntry.Name		  = fmt::format("{}#{}", HordeArtifactEntry.Name, HordeArtifactEntry.Id);
+			TResult<std::string> FormattedName = FormatArtifactName(HordeArtifactEntry, Query.Format);
 
-			Result.Entries.push_back(DirectoryEntry);
+			if (FormattedName.IsOk())
+			{
+				if (FormattedName->starts_with(RequiredPathPrefix))
+				{
+					std::string_view Name = std::string_view(*FormattedName).substr(RequiredPathPrefix.length());
+
+					DirectoryEntry.bDirectory = true;
+					DirectoryEntry.Name		  = Name;
+
+					Result.Entries.push_back(DirectoryEntry);
+				}
+			}
+			else
+			{
+				LogError(FormattedName.GetError());
+				break;
+			}
 		}
 	}
 	else
 	{
-		// Output virtual directories for unique CLs
-		uint64			  LastChangeNumber = 0;
 		for (const FHordeArtifactEntry& HordeArtifactEntry : HordeArtifacts)
 		{
-			if (LastChangeNumber == HordeArtifactEntry.Change)
-			{
-				continue;
-			}
-
 			FDirectoryListingEntry DirectoryEntry;
-
 			DirectoryEntry.bDirectory = true;
-			DirectoryEntry.Name		  = fmt::format("{}@{}", HordeArtifactEntry.StreamId, HordeArtifactEntry.Change);
-
+			DirectoryEntry.Name		  = fmt::format("#{}", HordeArtifactEntry.Id);
 			Result.Entries.push_back(DirectoryEntry);
+		}
+	}
+
+	// Create intermediate virtual directories
+	{
+		std::vector<FDirectoryListingEntry> VirtualEntries;
+		std::unordered_set<std::string>		UniqueEntries;
+
+		for (const FDirectoryListingEntry& Entry : Result.Entries)
+		{
+			UniqueEntries.insert(Entry.Name);
+		}
+
+		for (const FDirectoryListingEntry& Entry : Result.Entries)
+		{
+			std::vector<std::string_view> Parts = SplitBy(Entry.Name, PATH_SEPARATOR);
+			std::string					  PartialName;
+			for (std::string_view Part : Parts)
+			{
+				if (!PartialName.empty())
+				{
+					PartialName.push_back(PATH_SEPARATOR);
+				}
+				PartialName += Part;
 
-			LastChangeNumber = HordeArtifactEntry.Change;
+				bool bInserted = UniqueEntries.insert(PartialName).second;
+				if (bInserted)
+				{
+					FDirectoryListingEntry VirtualEntry;
+					VirtualEntry.Name = PartialName;
+					VirtualEntry.bDirectory = true;
+					VirtualEntries.push_back(VirtualEntry);
+				}
+			}
+		}
+
+		for (const FDirectoryListingEntry& Entry : VirtualEntries)
+		{
+			Result.Entries.push_back(Entry);
 		}
 	}
 
@@ -671,4 +866,23 @@ void TestHordeManifestDecode()
 	}
 }
 
+
+void
+TestHordeArtifactFormat()
+{
+	UNSYNC_LOG(L"TestHordeArtifactFormat()");
+	UNSYNC_LOG_INDENT;
+
+	FHordeArtifactEntry Artifact;
+	Artifact.Change	  = "12345";
+	Artifact.Id		  = "aabbcc";
+	Artifact.Name	  = "test_name";
+	Artifact.Metadata.push_back("foo=bar");
+	Artifact.Metadata.push_back("stream=test_stream");
+	Artifact.Metadata.push_back("buildname=build_name");
+
+	TResult<std::string> Name = FormatArtifactName(Artifact, "{Stream}.{BuildName}.{Change}");
+	UNSYNC_ASSERT(Name.GetData() == "test_stream.build_name.12345");
+}
+
 }
diff --git a/Engine/Source/Programs/Unsync/Private/UnsyncHorde.h b/Engine/Source/Programs/Unsync/Private/UnsyncHorde.h
--- a/Engine/Source/Programs/Unsync/Private/UnsyncHorde.h
+++ b/Engine/Source/Programs/Unsync/Private/UnsyncHorde.h
@@ -31,20 +31,27 @@ struct FHordeProtocolImpl : FRemoteProtocolBase
 };
 
 
-struct FHordeVirtualPath
+struct FHordeArtifactQuery
 {
-	std::string ArtifactPrefix;
+	std::string Type;
+	std::string Id;
+	std::string Change;
+	std::vector<std::string> Keys;
+
+	uint32 MaxResults = 1000000;
+
+	// Template to use for formatting results, e.g.: `{StreamName}.{BuildName}/{Platform}#{Id}`
+	std::string Format;
 
-	std::optional<std::string> StreamId;
-	std::optional<std::string> ArtifactId;
-	std::optional<uint64>	   Change;
+	// Query parts, delimited by directory separator character
+	std::vector<std::string> Parts;
 
-	static TResult<FHordeVirtualPath> FromString(std::string_view Str);
+	static TResult<FHordeArtifactQuery> FromString(std::string_view Str);
 };
 
 struct FHordeArtifactEntry
 {
-	uint64					 Change = 0;
+	std::string				 Change;
 	std::string				 Id;
 	std::string				 Name;
 	std::string				 Description;
diff --git a/Engine/Source/Programs/Unsync/Private/UnsyncMain.cpp b/Engine/Source/Programs/Unsync/Private/UnsyncMain.cpp
--- a/Engine/Source/Programs/Unsync/Private/UnsyncMain.cpp
+++ b/Engine/Source/Programs/Unsync/Private/UnsyncMain.cpp
@@ -669,21 +669,21 @@ InnerMain(int Argc, char** Argv)
 		bFilesystemSource = false;
 		if (RemoteDesc.RequestPath.empty())
 		{
-			TResult<FHordeVirtualPath> VirtualPath = FHordeVirtualPath::FromString(SourceFilenameUtf8);
-			if (VirtualPath.IsError())
+			TResult<FHordeArtifactQuery> Query = FHordeArtifactQuery::FromString(SourceFilenameUtf8);
+			if (Query.IsError())
 			{
 				UNSYNC_ERROR(L"Could not parse sync source path");
-				LogError(VirtualPath.GetError());
+				LogError(Query.GetError());
 				return 1;
 			}
 
-			if (!VirtualPath->ArtifactId)
+			if (Query->Id.empty())
 			{
 				UNSYNC_ERROR(L"Could not parse sync source path. Artifact ID is expected, i.e. '#123456abcdef'.");
 				return 1;
 			}
 
-			SourceFilenameUtf8 = "api/v2/artifacts/" + VirtualPath->ArtifactId.value();
+			SourceFilenameUtf8 = "api/v2/artifacts/" + Query->Id;
 			RemoteDesc.RequestPath = SourceFilenameUtf8;
 		}
 	}
diff --git a/Engine/Source/Programs/Unsync/Private/UnsyncProxy.cpp b/Engine/Source/Programs/Unsync/Private/UnsyncProxy.cpp
--- a/Engine/Source/Programs/Unsync/Private/UnsyncProxy.cpp
+++ b/Engine/Source/Programs/Unsync/Private/UnsyncProxy.cpp
@@ -850,7 +850,7 @@ ProxyQuery::FDirectoryListing::ToJson() const
 		}
 
 		Result += "{ ";
-		FormatJsonKeyValueStr(Result, "name", Entry.Name, ", ");
+		FormatJsonKeyValueStr(Result, "name", StringEscape(Entry.Name), ", ");
 		FormatJsonKeyValueBool(Result, "is_directory", Entry.bDirectory, ", ");
 		FormatJsonKeyValueUInt(Result, "mtime", Entry.Mtime, ", ");
 		FormatJsonKeyValueUInt(Result, "size", Entry.Size);
diff --git a/Engine/Source/Programs/Unsync/Private/UnsyncRemote.cpp b/Engine/Source/Programs/Unsync/Private/UnsyncRemote.cpp
--- a/Engine/Source/Programs/Unsync/Private/UnsyncRemote.cpp
+++ b/Engine/Source/Programs/Unsync/Private/UnsyncRemote.cpp
@@ -84,6 +84,13 @@ FRemoteDesc::FromUrl(std::string_view Url, EProtocolFlavor ProtocolFlavorHint)
 
 	FRemoteDesc Result;
 
+	const size_t NamespacePos = Url.find_last_of('#');
+	if (NamespacePos != std::string::npos)
+	{
+		Result.StorageNamespace = Url.substr(NamespacePos + 1);
+		Url						= Url.substr(0, NamespacePos);
+	}
+
 	const size_t SchemePos	 = Url.find("://");
 	const bool	 bHaveScheme = SchemePos != std::string::npos;
 
@@ -125,8 +132,6 @@ FRemoteDesc::FromUrl(std::string_view Url, EProtocolFlavor ProtocolFlavorHint)
 
 	const bool bRequestLooksLikeHordeArtifact = RequestPathLooksLikeHordeArtifact(Result.RequestPath);
 
-	const size_t NamespacePos = HostAddress.find_last_of('#');
-
 	if (ProtocolFlavorHint == EProtocolFlavor::Unknown)
 	{
 		// Try to guess protocol flavor
@@ -142,7 +147,7 @@ FRemoteDesc::FromUrl(std::string_view Url, EProtocolFlavor ProtocolFlavorHint)
 				{
 					Result.Protocol = EProtocolFlavor::Jupiter;
 				}
-				if (bRequestLooksLikeHordeArtifact || Scheme.starts_with("horde"))
+				else if (bRequestLooksLikeHordeArtifact || Scheme.starts_with("horde"))
 				{
 					Result.Protocol = EProtocolFlavor::Horde;
 				}
@@ -158,12 +163,6 @@ FRemoteDesc::FromUrl(std::string_view Url, EProtocolFlavor ProtocolFlavorHint)
 		Result.Protocol = ProtocolFlavorHint;
 	}
 
-	if (NamespacePos != std::string::npos)
-	{
-		Result.StorageNamespace = HostAddress.substr(NamespacePos + 1);
-		HostAddress				= HostAddress.substr(0, NamespacePos);
-	}
-
 	uint16		 HostPort = 0;
 	const size_t PortPos  = HostAddress.find_first_of(':');
 
diff --git a/Engine/Source/Programs/Unsync/Private/UnsyncTest.cpp b/Engine/Source/Programs/Unsync/Private/UnsyncTest.cpp
--- a/Engine/Source/Programs/Unsync/Private/UnsyncTest.cpp
+++ b/Engine/Source/Programs/Unsync/Private/UnsyncTest.cpp
@@ -952,6 +952,12 @@ RunTests(const std::string& Preset)
 		extern void TestHordeManifestDecode();
 		TestHordeManifestDecode();
 	}
+
+	if (Preset == "horde_artifact_format" || Preset == "all")
+	{
+		extern void TestHordeArtifactFormat();
+		TestHordeArtifactFormat();
+	}
 }
 
 }  // namespace unsync
diff --git a/Engine/Source/Programs/Unsync/Private/UnsyncUtil.cpp b/Engine/Source/Programs/Unsync/Private/UnsyncUtil.cpp
--- a/Engine/Source/Programs/Unsync/Private/UnsyncUtil.cpp
+++ b/Engine/Source/Programs/Unsync/Private/UnsyncUtil.cpp
@@ -345,30 +345,47 @@ ToString(const FPath& Path)
 }
 
 std::string
-StringToLower(const std::string& Input)
+StringToLower(std::string_view Input)
 {
-	std::string Result = Input;
+	std::string Result(Input);
 	std::transform(Result.begin(), Result.end(), Result.begin(), [](int32 C) { return char(::tolower(C)); });
 	return Result;
 }
 
 
 std::wstring
-StringToLower(const std::wstring& Input)
+StringToLower(std::wstring_view Input)
 {
-	std::wstring Result = Input;
+	std::wstring Result(Input);
 	std::transform(Result.begin(), Result.end(), Result.begin(), [](int32 C) { return wchar_t(::tolower(C)); });
 	return Result;
 }
 
 std::wstring
-StringToUpper(const std::wstring& Input)
+StringToUpper(std::wstring_view Input)
 {
-	std::wstring Result = Input;
+	std::wstring Result(Input);
 	std::transform(Result.begin(), Result.end(), Result.begin(), [](int32 C) { return wchar_t(::toupper(C)); });
 	return Result;
 }
 
+bool
+StringStartsWith(const std::string_view String, const std::string_view Prefix, bool bCaseSensitive)
+{
+	if (bCaseSensitive)
+	{
+		return String.starts_with(Prefix);
+	}
+	else if (Prefix.length() <= String.length())
+	{
+		return _strnicmp(String.data(), Prefix.data(), Prefix.length()) == 0;
+	}
+	else
+	{
+		return false;
+	}
+}
+
 bool
 StringEquals(const std::string_view A, const std::string_view B, bool bCaseSensitive)
 {
@@ -774,4 +791,28 @@ LooksLikeUrl(std::string_view Str)
 	return false;
 }
 
+std::vector<std::string_view>
+SplitByAny(std::string_view String, const char* SeparatorCharacters)
+{
+	std::vector<std::string_view> Result;
+
+	while (!String.empty())
+	{
+		size_t Pos = String.find_first_of(SeparatorCharacters);
+		if (Pos == std::string::npos)
+		{
+			Result.push_back(String);
+			break;
+		}
+
+		std::string_view Part = String.substr(0, Pos);
+
+		Result.push_back(Part);
+
+		String = String.substr(Pos + 1);
+	}
+
+	return Result;
+}
+
 }  // namespace unsync
diff --git a/Engine/Source/Programs/Unsync/Private/UnsyncUtil.h b/Engine/Source/Programs/Unsync/Private/UnsyncUtil.h
--- a/Engine/Source/Programs/Unsync/Private/UnsyncUtil.h
+++ b/Engine/Source/Programs/Unsync/Private/UnsyncUtil.h
@@ -154,13 +154,14 @@ ToString(const std::string_view StringView)
 
 std::string ToString(const FPath& Path);
 
-std::string StringToLower(const std::string& Input);
-std::wstring StringToLower(const std::wstring& Input);
-std::wstring StringToUpper(const std::wstring& Input);
+std::string StringToLower(std::string_view Input);
+std::wstring StringToLower(std::wstring_view Input);
+std::wstring StringToUpper(std::wstring_view Input);
 
 std::string StringEscape(const std::string_view Input);
 
 bool StringEquals(const std::string_view A, const std::string_view B, bool bCaseSensitive = true);
+bool StringStartsWith(const std::string_view String, const std::string_view Prefix, bool bCaseSensitive = true);
 
 inline bool
 UncasedStringEquals(const std::string_view A, const std::string_view B)
@@ -168,6 +169,33 @@ UncasedStringEquals(const std::string_view A, const std::string_view B)
 	return StringEquals(A, B, false);
 }
 
+inline bool
+IsAsciiAlphabetCharacter(const char C)
+{
+	return (C >= 'a' && C <= 'z') || (C >= 'A' && C <= 'Z');
+}
+
+inline bool
+IsAsciiNumericCharacter(const char C)
+{
+	return C >= '0' && C <= '9';
+}
+
+inline bool
+IsAsciiAlphaNumericCharacter(const char C)
+{
+	return IsAsciiAlphabetCharacter(C) || IsAsciiNumericCharacter(C);
+}
+
+std::vector<std::string_view> SplitByAny(std::string_view String, const char* SeparatorCharacters);
+inline std::vector<std::string_view>
+
+SplitBy(std::string_view String, char SeparatorCharacter)
+{
+	char Chars[2] = {SeparatorCharacter, 0};
+	return SplitByAny(String, Chars);
+}
+
 // Takes a drive-based path (e.g. P:/Foo/Bar) and converts it to universal form (e.g. //server/Foo/Bar), if possible.
 // Otherwise, returns original path.
 FPath GetUniversalPath(const FPath& Path);
diff --git a/Engine/Source/Programs/Unsync/Private/UnsyncVersion.h b/Engine/Source/Programs/Unsync/Private/UnsyncVersion.h
--- a/Engine/Source/Programs/Unsync/Private/UnsyncVersion.h
+++ b/Engine/Source/Programs/Unsync/Private/UnsyncVersion.h
@@ -4,7 +4,7 @@
 
 #include "UnsyncCommon.h"
 
-#define UNSYNC_VERSION_STR "1.0.79"
+#define UNSYNC_VERSION_STR "1.0.80"
 
 namespace unsync {
 
