unsync - Fix compile error

diff --git a/Engine/Source/Programs/Unsync/Private/UnsyncCommon.h b/Engine/Source/Programs/Unsync/Private/UnsyncCommon.h
--- a/Engine/Source/Programs/Unsync/Private/UnsyncCommon.h
+++ b/Engine/Source/Programs/Unsync/Private/UnsyncCommon.h
@@ -24,6 +24,8 @@ using int64 = int64_t;
 using FPath = std::filesystem::path;
 using FPathStringView = std::basic_string_view<FPath::value_type>;
 
+static const char PATH_SEPARATOR = char(FPath::preferred_separator);
+
 #define UNSYNC_DISALLOW_COPY_ASSIGN(T) \
 	T(const T&) = delete;              \
 	T& operator=(const T&) = delete;
diff --git a/Engine/Source/Programs/Unsync/Private/UnsyncFile.cpp b/Engine/Source/Programs/Unsync/Private/UnsyncFile.cpp
--- a/Engine/Source/Programs/Unsync/Private/UnsyncFile.cpp
+++ b/Engine/Source/Programs/Unsync/Private/UnsyncFile.cpp
@@ -144,7 +144,7 @@ ConvertDirectorySeparatorsToNative(std::string& Path)
 		Path.begin(),
 		Path.end(),
 		[](char C) { return C == '/' || C == '\\'; },
-		char(FPath::preferred_separator));
+		PATH_SEPARATOR);
 }
 
 void
