unsync - Fix horde query parsing when it only contains the artifact id

diff --git a/Engine/Source/Programs/Unsync/Private/UnsyncHorde.cpp b/Engine/Source/Programs/Unsync/Private/UnsyncHorde.cpp
--- a/Engine/Source/Programs/Unsync/Private/UnsyncHorde.cpp
+++ b/Engine/Source/Programs/Unsync/Private/UnsyncHorde.cpp
@@ -330,7 +330,7 @@ FHordeArtifactQuery::FromString(std::string_view Str)
 
 	for (std::string_view PathPart : PathParts)
 	{
-		if (!Result.Parts.empty())
+		if (!Result.Parts.empty() || PathParts.size() == 1)
 		{
 			// TODO: derived the ID by parsing the request string using the format specification
 			size_t ArtifactIdPos = PathPart.find_last_of('#');
