unsync - Add `query http-get <url>` command to easily issue use Horde and Unsync HTTP APIs

* Add extended error messages to various places

diff --git a/Engine/Source/Programs/Unsync/Private/UnsyncCmdLogin.cpp b/Engine/Source/Programs/Unsync/Private/UnsyncCmdLogin.cpp
--- a/Engine/Source/Programs/Unsync/Private/UnsyncCmdLogin.cpp
+++ b/Engine/Source/Programs/Unsync/Private/UnsyncCmdLogin.cpp
@@ -79,7 +79,7 @@ CmdLogin(const FCmdLoginOptions& Options)
 			}
 			else
 			{
-				LogError(HttpError(Response.Code));
+				LogError(HttpError(Response.Code), L"Failed to authenticate");
 				return -1;
 			}
 		}
@@ -120,7 +120,7 @@ CmdLogin(const FCmdLoginOptions& Options)
 	}
 	else
 	{
-		LogError(AuthTokenResult.GetError());
+		LogError(AuthTokenResult.GetError(), L"Failed to authenticate");
 		return -1;
 	}
 }
diff --git a/Engine/Source/Programs/Unsync/Private/UnsyncCmdQuery.cpp b/Engine/Source/Programs/Unsync/Private/UnsyncCmdQuery.cpp
--- a/Engine/Source/Programs/Unsync/Private/UnsyncCmdQuery.cpp
+++ b/Engine/Source/Programs/Unsync/Private/UnsyncCmdQuery.cpp
@@ -127,7 +127,7 @@ CmdQueryMirrors(const FCmdQueryOptions& Options)
 	FMirrorInfoResult MirrorsResult = RunQueryMirrors(Options.Remote);
 	if (MirrorsResult.IsError())
 	{
-		LogError(MirrorsResult.GetError());
+		LogError(MirrorsResult.GetError(), L"Failed to get mirror list from the server");
 		return 1;
 	}
 
@@ -189,7 +189,7 @@ CmdQueryList(const FCmdQueryOptions& Options)
 
 	if (ListingResult.IsError())
 	{
-		LogError(ListingResult.GetError());
+		LogError(ListingResult.GetError(), L"Failed to list remote directory");
 		return -1;
 	}
 
@@ -313,7 +313,7 @@ CmdQuerySearch(const FCmdQueryOptions& Options)
 
 			if (DirectoryListingResult.IsError())
 			{
-				LogError(DirectoryListingResult.GetError());
+				LogError(DirectoryListingResult.GetError(), L"Failed to list remote directory");
 				return;
 			}
 
@@ -451,7 +451,7 @@ CmdQueryFile(const FCmdQueryOptions& Options)
 
 	if (!AuthToken.IsOk())
 	{
-		LogError(AuthToken.GetError());
+		LogError(AuthToken.GetError(), L"Failed to authenticate");
 		return -1;
 	}
 
@@ -496,12 +496,98 @@ CmdQueryFile(const FCmdQueryOptions& Options)
 	}
 	else
 	{
-		LogError(Response.GetError());
+		LogError(Response.GetError(), L"Failed to download file");
 	}
 
 	return 0;
 }
 
+int32
+CmdQueryHttpGet(const FCmdQueryOptions& Options)
+{
+	FHttpConnection HttpConnection = FHttpConnection::CreateDefaultHttps(Options.Remote);
+
+	std::string BearerToken;
+
+	if (Options.Remote.bAuthenticationRequired)
+	{
+		TResult<ProxyQuery::FHelloResponse> HelloResponse = ProxyQuery::Hello(Options.Remote.Protocol, HttpConnection);
+		if (!HelloResponse.IsOk())
+		{
+			LogError(HelloResponse.GetError(), L"Failed to query basic server information");
+			return -1;
+		}
+
+		FAuthDesc			AuthDesc  = FAuthDesc::FromHelloResponse(*HelloResponse);
+		TResult<FAuthToken> AuthToken = Authenticate(AuthDesc);
+
+		if (!AuthToken.IsOk())
+		{
+			LogError(AuthToken.GetError(), L"Failed to authenticate with the server");
+			return -1;
+		}
+
+		BearerToken = AuthToken->Access;
+	}
+
+	// TODO: RequestPath should include leading slash
+	std::string RequestUrl = fmt::format("/{}", Options.Remote.RequestPath);
+
+	FHttpRequest Request;
+	Request.Method		= EHttpMethod::GET;
+	Request.BearerToken = BearerToken;
+	Request.Url			= RequestUrl;
+
+	FHttpResponse Response = HttpRequest(HttpConnection, Request);
+
+	if (!Response.Success())
+	{
+		LogError(HttpError(std::move(RequestUrl), Response.Code));
+	}
+
+	FPath OutputPath = Options.OutputPath;
+	if (OutputPath.empty())
+	{
+		if (Response.ContentType == EHttpContentType::Application_Json || Response.ContentType == EHttpContentType::Text_Plain ||
+			Response.ContentType == EHttpContentType::Text_Html)
+		{
+			Response.Buffer.PushBack('\n');
+			Response.Buffer.PushBack(0);
+			LogPrintf(ELogLevel::MachineReadable, L"%hs", (const char*)Response.Buffer.Data());
+			return 0;
+		}
+		else
+		{
+			UNSYNC_ERROR(L"Unexpected response content type. Only plain text or json are supported. Use `-o <filename>` command line argument to write response body to a file.");
+			return -1;
+		}
+	}
+	else
+	{
+		UNSYNC_LOG(L"Output file: '%ls'", OutputPath.wstring().c_str());
+		OutputPath = GetAbsoluteNormalPath(OutputPath);
+
+		if (EnsureDirectoryExists(OutputPath.parent_path()))
+		{
+			if (WriteBufferToFile(OutputPath, Response.Buffer))
+			{
+				UNSYNC_VERBOSE(L"Wrote bytes: %llu", llu(Response.Buffer.Size()));
+				return 0;
+			}
+			else
+			{
+				UNSYNC_ERROR(L"Failed to write output file '%ls'", OutputPath.wstring().c_str());
+			}
+		}
+		else
+		{
+			UNSYNC_ERROR(L"Failed to create output directory '%ls'", OutputPath.parent_path().wstring().c_str());
+		}
+
+		return -1;
+	}
+}
+
 int32
 CmdQuery(const FCmdQueryOptions& Options)
 {
@@ -527,9 +613,13 @@ CmdQuery(const FCmdQueryOptions& Options)
 	{
 		return CmdQueryFile(Options);
 	}
+	if (Options.Query == "http-get")
+	{
+		return CmdQueryHttpGet(Options);
+	}
 	else
 	{
-		UNSYNC_ERROR(L"Unknown query command");
+		UNSYNC_ERROR(L"Unknown query command. Allowed options: mirrors, list, search, file, http-get");
 		return 1;
 	}
 }
diff --git a/Engine/Source/Programs/Unsync/Private/UnsyncCore.cpp b/Engine/Source/Programs/Unsync/Private/UnsyncCore.cpp
--- a/Engine/Source/Programs/Unsync/Private/UnsyncCore.cpp
+++ b/Engine/Source/Programs/Unsync/Private/UnsyncCore.cpp
@@ -955,7 +955,7 @@ CopyFileIfPossiblyDifferent(FProxyFileSystem&	   FileSystem,
 		TResult<FBuffer> FileBuffer = FileSystem.ReadFile(ToString(Source.Path));
 		if (FileBuffer.IsError())
 		{
-			LogError(FileBuffer.GetError());
+			LogError(FileBuffer.GetError(), L"Failed to read source file");
 			return false;
 		}
 
@@ -1014,7 +1014,7 @@ LoadAndMergeSourceManifest(FDirectoryManifest& Output,
 	TResult<FFoundManifest> FindManifestResult = FindUnsyncManifest(ProxyFileSystem);
 	if (FindManifestResult.IsError())
 	{
-		LogError(FindManifestResult.GetError());
+		LogError(FindManifestResult.GetError(), L"Failed to find remote manifest");
 		return false;
 	}
 
@@ -1058,7 +1058,7 @@ LoadAndMergeSourceManifest(FDirectoryManifest& Output,
 			TResult<FBuffer> FileBuffer = ProxyFileSystem.ReadFile(ToString(IndexFileInfo.Path));
 			if (FileBuffer.IsError())
 			{
-				LogError(FileBuffer.GetError());
+				LogError(FileBuffer.GetError(), L"Failed to read remote file");
 				return false;
 			}
 
@@ -1367,7 +1367,7 @@ SyncDirectory(const FSyncDirectoryOptions& SyncOptions)
 		}
 		else
 		{
-			LogError(DownloadResult.GetError());
+			LogError(DownloadResult.GetError(), L"Failed to download file");
 			UNSYNC_BREAK_ON_ERROR;
 			return false;
 		}
diff --git a/Engine/Source/Programs/Unsync/Private/UnsyncHorde.cpp b/Engine/Source/Programs/Unsync/Private/UnsyncHorde.cpp
--- a/Engine/Source/Programs/Unsync/Private/UnsyncHorde.cpp
+++ b/Engine/Source/Programs/Unsync/Private/UnsyncHorde.cpp
@@ -608,7 +608,7 @@ FHordeProtocolImpl::QueryListDirectory(FHttpConnection& Connection, const FAuthD
 			}
 			else
 			{
-				LogError(FormattedName.GetError());
+				LogError(FormattedName.GetError(), L"Failed to format Horde artifact name");
 				break;
 			}
 		}
@@ -862,7 +862,7 @@ void TestHordeManifestDecode()
 	TResult<FDirectoryManifest> Manifest = DecodeHordeManifestJson(GTestHordeManifestJson, "api/v2/artifacts/12345");
 	if (Manifest.IsError())
 	{
-		LogError(Manifest.GetError());
+		LogError(Manifest.GetError(), L"Failed to decode Horde manifest from JSON");
 	}
 }
 
diff --git a/Engine/Source/Programs/Unsync/Private/UnsyncLog.cpp b/Engine/Source/Programs/Unsync/Private/UnsyncLog.cpp
--- a/Engine/Source/Programs/Unsync/Private/UnsyncLog.cpp
+++ b/Engine/Source/Programs/Unsync/Private/UnsyncLog.cpp
@@ -381,7 +381,7 @@ LogPrintf(ELogLevel Level, const wchar_t* Str, ...)
 
 extern const char* HttpStatusToString(int32 Code);
 void
-LogError(const FError& E)
+LogError(const FError& E, std::wstring ExtraContext)
 {
 	const char* ErrorKindStr = nullptr;
 	const char* ErrorDescStr = nullptr;
@@ -403,25 +403,30 @@ LogError(const FError& E)
 			break;
 	}
 
-	const wchar_t* ContextStr = E.Context.empty() ? nullptr : E.Context.c_str();
+	const bool bHaveContext = !E.Context.empty();
+	const bool bHaveExtraContext = !ExtraContext.empty();
 
 	if (ErrorDescStr)
 	{
 		LogPrintf(ELogLevel::Error,
-				  L"%hs code: %d (%hs).%ls%ls\n",
+				  L"%ls%hs%hs code: %d (%hs).%ls%ls\n",
+				  bHaveExtraContext ? ExtraContext.c_str() : L"",
+				  bHaveExtraContext ? ": " : "",
 				  ErrorKindStr,
 				  E.Code,
 				  ErrorDescStr,
-				  ContextStr ? L" Context: " : L"",
+				  bHaveContext ? L" Context: " : L"",
 				  E.Context.empty() ? L"" : E.Context.c_str());
 	}
 	else
 	{
 		LogPrintf(ELogLevel::Error,
-				  L"%hs code: %d.%ls%ls\n",
+				  L"%ls%hs%hs code: %d.%ls%ls\n",
+				  bHaveExtraContext ? ExtraContext.c_str() : L"",
+				  bHaveExtraContext ? ": " : "",
 				  ErrorKindStr,
 				  E.Code,
-				  ContextStr ? L" Context: " : L"",
+				  bHaveContext ? L" Context: " : L"",
 				  E.Context.empty() ? L"" : E.Context.c_str());
 	}
 }
diff --git a/Engine/Source/Programs/Unsync/Private/UnsyncLog.h b/Engine/Source/Programs/Unsync/Private/UnsyncLog.h
--- a/Engine/Source/Programs/Unsync/Private/UnsyncLog.h
+++ b/Engine/Source/Programs/Unsync/Private/UnsyncLog.h
@@ -74,7 +74,7 @@ enum class ELogLevel
 void LogFlush();
 
 void LogPrintf(ELogLevel Level, const wchar_t* Str, ...);
-void LogError(const FError& E);
+void LogError(const FError& E, std::wstring ExtraContext = {});
 
 void LogProgress(const wchar_t* ItemName, uint64 Current, uint64 Total);
 void LogStatus(const wchar_t* ItemName, const wchar_t* Status);
diff --git a/Engine/Source/Programs/Unsync/Private/UnsyncMain.cpp b/Engine/Source/Programs/Unsync/Private/UnsyncMain.cpp
--- a/Engine/Source/Programs/Unsync/Private/UnsyncMain.cpp
+++ b/Engine/Source/Programs/Unsync/Private/UnsyncMain.cpp
@@ -672,8 +672,7 @@ InnerMain(int Argc, char** Argv)
 			TResult<FHordeArtifactQuery> Query = FHordeArtifactQuery::FromString(SourceFilenameUtf8);
 			if (Query.IsError())
 			{
-				UNSYNC_ERROR(L"Could not parse sync source path");
-				LogError(Query.GetError());
+				LogError(Query.GetError(), L"Could not parse sync source path");
 				return 1;
 			}
 
@@ -893,7 +892,8 @@ InnerMain(int Argc, char** Argv)
 
 	if (bShouldLogin)
 	{
-		RemoteDesc.PrimaryHost = RemoteDesc.Host;
+		RemoteDesc.PrimaryHost			   = RemoteDesc.Host;
+		RemoteDesc.bAuthenticationRequired = true;
 	}
 
 	FRemoteDesc RootRemoteDesc = RemoteDesc;
diff --git a/Engine/Source/Programs/Unsync/Private/UnsyncProxy.cpp b/Engine/Source/Programs/Unsync/Private/UnsyncProxy.cpp
--- a/Engine/Source/Programs/Unsync/Private/UnsyncProxy.cpp
+++ b/Engine/Source/Programs/Unsync/Private/UnsyncProxy.cpp
@@ -217,8 +217,7 @@ FUnsyncProtocolImpl::FUnsyncProtocolImpl(const FRemoteDesc&				RemoteDesc,
 			}
 			else
 			{
-				LogError(AuthTokenResult.GetError());
-				UNSYNC_ERROR(L"Server requires authentication, but access token could not be acquired");
+				LogError(AuthTokenResult.GetError(), L"Server requires authentication, but access token could not be acquired");
 				Invalidate();
 			}
 		}
@@ -1305,7 +1304,7 @@ FProxyPool::FProxyPool(const FRemoteDesc& InRemoteDesc, const FAuthDesc* InAuthD
 
 		if (Response.IsError())
 		{
-			LogError(Response.GetError());
+			LogError(Response.GetError(), L"Failed to query basic server information");
 		}
 		else
 		{
@@ -1403,7 +1402,7 @@ FProxyPool::GetAccessToken()
 		}
 		else
 		{
-			LogError(AuthTokenResult.GetError());
+			LogError(AuthTokenResult.GetError(), L"Failed to authenticate");
 			UNSYNC_FATAL(L"Cannot proceed without a valid authentication token");
 		}
 	}
diff --git a/Engine/Source/Programs/Unsync/Private/UnsyncVersion.h b/Engine/Source/Programs/Unsync/Private/UnsyncVersion.h
--- a/Engine/Source/Programs/Unsync/Private/UnsyncVersion.h
+++ b/Engine/Source/Programs/Unsync/Private/UnsyncVersion.h
@@ -4,7 +4,7 @@
 
 #include "UnsyncCommon.h"
 
-#define UNSYNC_VERSION_STR "1.0.80"
+#define UNSYNC_VERSION_STR "1.0.81"
 
 namespace unsync {
 
