unsync - Add early exit if non-success HTTP code is encountered during `query http-get` command

* Update version to 1.0.82

diff --git a/Engine/Source/Programs/Unsync/Private/UnsyncCmdQuery.cpp b/Engine/Source/Programs/Unsync/Private/UnsyncCmdQuery.cpp
--- a/Engine/Source/Programs/Unsync/Private/UnsyncCmdQuery.cpp
+++ b/Engine/Source/Programs/Unsync/Private/UnsyncCmdQuery.cpp
@@ -543,12 +543,17 @@ CmdQueryHttpGet(const FCmdQueryOptions& Options)
 	if (!Response.Success())
 	{
 		LogError(HttpError(std::move(RequestUrl), Response.Code));
+		return -1;
 	}
 
 	FPath OutputPath = Options.OutputPath;
 	if (OutputPath.empty())
 	{
-		if (Response.ContentType == EHttpContentType::Application_Json || Response.ContentType == EHttpContentType::Text_Plain ||
+		if (Response.Buffer.Empty())
+		{
+			return 0;
+		}
+		else if (Response.ContentType == EHttpContentType::Application_Json || Response.ContentType == EHttpContentType::Text_Plain ||
 			Response.ContentType == EHttpContentType::Text_Html)
 		{
 			Response.Buffer.PushBack('\n');
diff --git a/Engine/Source/Programs/Unsync/Private/UnsyncVersion.h b/Engine/Source/Programs/Unsync/Private/UnsyncVersion.h
--- a/Engine/Source/Programs/Unsync/Private/UnsyncVersion.h
+++ b/Engine/Source/Programs/Unsync/Private/UnsyncVersion.h
@@ -4,7 +4,7 @@
 
 #include "UnsyncCommon.h"
 
-#define UNSYNC_VERSION_STR "1.0.81"
+#define UNSYNC_VERSION_STR "1.0.82"
 
 namespace unsync {
 
