unsync - Fix manifest caching when using overlays

* Virtualized file system was created for the main sync source and erroneously was used when attempting to fetch overlay manifest

diff --git a/Engine/Source/Programs/Unsync/Private/UnsyncCore.cpp b/Engine/Source/Programs/Unsync/Private/UnsyncCore.cpp
--- a/Engine/Source/Programs/Unsync/Private/UnsyncCore.cpp
+++ b/Engine/Source/Programs/Unsync/Private/UnsyncCore.cpp
@@ -1386,30 +1386,6 @@ SyncDirectory(const FSyncDirectoryOptions& SyncOptions)
 	}
 	else
 	{
-		std::unique_ptr<FProxyFileSystem> ProxyFileSystem;
-		if (bServerSource)
-		{
-			const FRemoteProtocolFeatures& RemoteFeatures = ProxyPool.GetFeatures();
-
-			if (!RemoteFeatures.bDirectoryListing)
-			{
-				UNSYNC_ERROR(L"Remote server does not support directory listing");
-				return false;
-			}
-
-			if (!RemoteFeatures.bFileDownload)
-			{
-				UNSYNC_ERROR(L"Remote server does not support direct file downloads");
-				return false;
-			}
-
-			ProxyFileSystem = std::make_unique<FRemoteFileSystem>(ToString(SyncOptions.Source), ProxyPool);
-		}
-		else
-		{
-			ProxyFileSystem = std::make_unique<FPhysicalFileSystem>(SyncOptions.Source);
-		}
-
 		std::vector<FPath> AllSources;
 		AllSources.push_back(SourcePath);
 		for (const FPath& OverlayPath : SyncOptions.Overlays)
@@ -1419,6 +1395,30 @@ SyncDirectory(const FSyncDirectoryOptions& SyncOptions)
 
 		for (const FPath& ThisSourcePath : AllSources)
 		{
+			std::unique_ptr<FProxyFileSystem> ProxyFileSystem;
+			if (bServerSource)
+			{
+				const FRemoteProtocolFeatures& RemoteFeatures = ProxyPool.GetFeatures();
+
+				if (!RemoteFeatures.bDirectoryListing)
+				{
+					UNSYNC_ERROR(L"Remote server does not support directory listing");
+					return false;
+				}
+
+				if (!RemoteFeatures.bFileDownload)
+				{
+					UNSYNC_ERROR(L"Remote server does not support direct file downloads");
+					return false;
+				}
+
+				ProxyFileSystem = std::make_unique<FRemoteFileSystem>(ToString(ThisSourcePath), ProxyPool);
+			}
+			else
+			{
+				ProxyFileSystem = std::make_unique<FPhysicalFileSystem>(ThisSourcePath);
+			}
+
 			if (!LoadAndMergeSourceManifest(SourceDirectoryManifest,
 											PackIndexFiles,
 											*ProxyFileSystem,
diff --git a/Engine/Source/Programs/Unsync/Private/UnsyncVersion.h b/Engine/Source/Programs/Unsync/Private/UnsyncVersion.h
--- a/Engine/Source/Programs/Unsync/Private/UnsyncVersion.h
+++ b/Engine/Source/Programs/Unsync/Private/UnsyncVersion.h
@@ -4,7 +4,7 @@
 
 #include "UnsyncCommon.h"
 
-#define UNSYNC_VERSION_STR "1.0.82"
+#define UNSYNC_VERSION_STR "1.0.83"
 
 namespace unsync {
 
