unsync - Increase default socket timeout from 1 minute to 15

diff --git a/Engine/Source/Programs/Unsync/Private/UnsyncMain.cpp b/Engine/Source/Programs/Unsync/Private/UnsyncMain.cpp
--- a/Engine/Source/Programs/Unsync/Private/UnsyncMain.cpp
+++ b/Engine/Source/Programs/Unsync/Private/UnsyncMain.cpp
@@ -345,7 +345,7 @@ InnerMain(int Argc, char** Argv)
 		->group(ExperimentalGroupId);
 	SubSync->add_flag("--login", bShouldLogin, "Use user authentication when accessing unsync server");
 	SubSync->add_option("--token", AuthTokenPathUtf8, "Explicit path to the authentication token file to use");
-	SubSync->add_flag("--no-timeout", bNoSocketTimeout, "Disable the default 60 second timeout on network socket operations");
+	SubSync->add_flag("--no-timeout", bNoSocketTimeout, "Disable the default 15 minute timeout on network socket operations");
 
 	CLI::Option* BackgroundMemoryBudgetOption = SubSync->add_option("--background-task-memory",
 															BackgroundTaskMemoryBudgetGB,
@@ -1034,7 +1034,7 @@ InnerMain(int Argc, char** Argv)
 		}
 		else
 		{
-			RemoteDesc.RecvTimeoutSeconds = 60;
+			RemoteDesc.RecvTimeoutSeconds = 15 * 60;
 		}
 
 		// Try to derive default memory budget
diff --git a/Engine/Source/Programs/Unsync/Private/UnsyncVersion.h b/Engine/Source/Programs/Unsync/Private/UnsyncVersion.h
--- a/Engine/Source/Programs/Unsync/Private/UnsyncVersion.h
+++ b/Engine/Source/Programs/Unsync/Private/UnsyncVersion.h
@@ -4,7 +4,7 @@
 
 #include "UnsyncCommon.h"
 
-#define UNSYNC_VERSION_STR "1.0.83"
+#define UNSYNC_VERSION_STR "1.0.84"
 
 namespace unsync {
 
