unsync - Implement support for overlays when syncing from Horde

diff --git a/Engine/Source/Programs/Unsync/Private/UnsyncCmdSync.cpp b/Engine/Source/Programs/Unsync/Private/UnsyncCmdSync.cpp
--- a/Engine/Source/Programs/Unsync/Private/UnsyncCmdSync.cpp
+++ b/Engine/Source/Programs/Unsync/Private/UnsyncCmdSync.cpp
@@ -22,6 +22,27 @@ CmdSync(const FCmdSyncOptions& Options)
 		return Options.Filter ? Options.Filter->Resolve(Path) : Path;
 	};
 
+	if (Options.Source.empty())
+	{
+		UNSYNC_ERROR(L"Sync source location is invalid");
+		return 1;
+	}
+
+	if (Options.Target.empty())
+	{
+		UNSYNC_ERROR(L"Sync target location is invalid");
+		return 1;
+	}
+
+	for (const FPath& Path : Options.Overlays)
+	{
+		if (Path.empty())
+		{
+			UNSYNC_ERROR(L"Sync overlay location is invalid");
+			return 1;
+		}
+	}
+
 	FProxyPool ProxyPool(Options.Remote, Options.AuthDesc);
 
 	std::error_code ErrorCode	   = {};
diff --git a/Engine/Source/Programs/Unsync/Private/UnsyncCore.cpp b/Engine/Source/Programs/Unsync/Private/UnsyncCore.cpp
--- a/Engine/Source/Programs/Unsync/Private/UnsyncCore.cpp
+++ b/Engine/Source/Programs/Unsync/Private/UnsyncCore.cpp
@@ -1335,8 +1335,6 @@ SyncDirectory(const FSyncDirectoryOptions& SyncOptions)
 		}
 	};
 
-	auto ResolvePath = [SyncFilter](const FPath& Filename) -> FPath { return SyncFilter ? SyncFilter->Resolve(Filename) : Filename; };
-
 	FDirectoryManifest SourceDirectoryManifest;
 	FPath			   SourceManifestTempPath;
 
@@ -1346,6 +1344,18 @@ SyncDirectory(const FSyncDirectoryOptions& SyncOptions)
 
 	std::vector<FPackIndexDatabase> PackIndexFiles;
 
+	std::vector<FPath> AllSources;
+	AllSources.push_back(SourcePath);
+	for (const FPath& OverlayPath : SyncOptions.Overlays)
+	{
+		AllSources.push_back(OverlayPath);
+	}
+	
+	auto ResolvePath = [SyncFilter](const FPath& Filename) -> FPath { return SyncFilter ? SyncFilter->Resolve(Filename) : Filename; };
+
+	// Used to build block request map when syncing from multiple sources.
+	THashMap<FHash256, uint32> FileSourceIdMap;
+
 	if (SyncOptions.SourceType == ESourceType::ServerWithManifestId)
 	{
 		if (!ProxyPool.IsValid())
@@ -1354,22 +1364,59 @@ SyncDirectory(const FSyncDirectoryOptions& SyncOptions)
 			return false;
 		}
 
-		UNSYNC_LOG(L"Downloading manifest ...");
-
 		std::unique_ptr<FProxy> Proxy = ProxyPool.Alloc();
 
-		std::string					SourceManifestName = ConvertWideToUtf8(SyncOptions.Source.wstring());
-		TResult<FDirectoryManifest> DownloadResult	   = Proxy->DownloadManifest(SourceManifestName);
-		if (FDirectoryManifest* Manifest = DownloadResult.TryData())
-		{
-			bSourceManifestOk = true;
-			std::swap(SourceDirectoryManifest, *Manifest);
-		}
-		else
+		uint32 SourceIndex = 0;
+		for (const FPath& ThisSourcePath : AllSources)
 		{
-			LogError(DownloadResult.GetError(), L"Failed to download file");
-			UNSYNC_BREAK_ON_ERROR;
-			return false;
+			std::string SourceManifestName = ConvertWideToUtf8(ThisSourcePath.wstring());
+			FHash128	SourcePathHash	   = HashBlake3String<FHash128>(SourceManifestName);
+			std::string SourcePathHashStr  = BytesToHexString(SourcePathHash.Data, sizeof(SourcePathHash.Data));
+			FPath		CachedManifestPath = TargetTempPath / SourcePathHashStr;
+
+			FPath EmptyPath; // no physical path for downloaded manifests
+
+			FDirectoryManifest LoadedManifest;
+			if (!PathExists(CachedManifestPath) || !LoadDirectoryManifest(LoadedManifest, EmptyPath, CachedManifestPath))
+			{
+				LogGlobalStatus(L"Caching source manifest");
+				UNSYNC_VERBOSE(L"Caching source manifest");
+
+				UNSYNC_LOG_INDENT;
+				UNSYNC_VERBOSE(L"Source '%hs'", SourceManifestName.c_str());
+				UNSYNC_VERBOSE(L"Target '%ls'", CachedManifestPath.wstring().c_str());
+
+				TResult<FDirectoryManifest> DownloadResult = Proxy->DownloadManifest(SourceManifestName);
+
+				if (FDirectoryManifest* Manifest = DownloadResult.TryData())
+				{
+					std::swap(LoadedManifest, *Manifest);
+				}
+				else
+				{
+					LogError(DownloadResult.GetError(), L"Failed to download manifest");
+					UNSYNC_BREAK_ON_ERROR;
+					return false;
+				}
+
+				const bool bAllowInDryRun = true;
+				SaveDirectoryManifest(LoadedManifest, CachedManifestPath, bAllowInDryRun);
+			}
+
+			for (const auto& It : LoadedManifest.Files)
+			{
+				FHash256 NameHash = HashBlake3String<FHash256>(It.first);
+				FileSourceIdMap[NameHash] = SourceIndex;
+			}
+
+			bSourceManifestOk = MergeManifests(SourceDirectoryManifest, LoadedManifest, bCaseSensitiveTargetFileSystem);
+
+			if (!bSourceManifestOk)
+			{
+				break;
+			}
+
+			++SourceIndex;
 		}
 
 		ProxyPool.Dealloc(std::move(Proxy)); // TODO: RAII helper for pooled proxy connections
@@ -1386,13 +1433,6 @@ SyncDirectory(const FSyncDirectoryOptions& SyncOptions)
 	}
 	else
 	{
-		std::vector<FPath> AllSources;
-		AllSources.push_back(SourcePath);
-		for (const FPath& OverlayPath : SyncOptions.Overlays)
-		{
-			AllSources.push_back(OverlayPath);
-		}
-
 		for (const FPath& ThisSourcePath : AllSources)
 		{
 			std::unique_ptr<FProxyFileSystem> ProxyFileSystem;
@@ -1528,6 +1568,7 @@ SyncDirectory(const FSyncDirectoryOptions& SyncOptions)
 	for (const auto& SourceManifestIt : SourceDirectoryManifest.Files)
 	{
 		const std::wstring& SourceFilename = SourceManifestIt.first;
+		const FHash256		SourceFilenameHash = HashBlake3String<FHash256>(SourceFilename);
 
 		if (!ShouldSync(SourceFilename))
 		{
@@ -1564,6 +1605,16 @@ SyncDirectory(const FSyncDirectoryOptions& SyncOptions)
 		FPath TargetFilePath = TargetPath / ToPath(SourceManifestIt.first);
 
 		FPath ResolvedSourceFilePath = ResolvePath(SourceFilePath);
+		uint32 SourceId = 0;
+
+		{
+			FHash256 NameHash	   = HashBlake3String<FHash256>(SourceFilename);
+			auto	 FoundSourceId = FileSourceIdMap.find(NameHash);
+			if (FoundSourceId != FileSourceIdMap.end())
+			{
+				SourceId = FoundSourceId->second;
+			}
+		}
 
 		if (bFileSystemSource && SyncOptions.bValidateSourceFiles)
 		{
@@ -1633,6 +1684,7 @@ SyncDirectory(const FSyncDirectoryOptions& SyncOptions)
 			Task.BaseFilePath			= std::move(BaseFilePath);
 			Task.TargetFilePath			= std::move(TargetFilePath);
 			Task.SourceManifest			= &SourceFileManifest;
+			Task.SourceId				= SourceId;
 
 			if (bQuickDifferencePossible)
 			{
@@ -1815,11 +1867,11 @@ SyncDirectory(const FSyncDirectoryOptions& SyncOptions)
 			if (bProxyHasData)
 			{
 				FBlockRequestMap BlockRequestMap;
-				BlockRequestMap.Init(SourceDirectoryManifest.Algorithm.StrongHashAlgorithmId);
+				BlockRequestMap.Init(SourceDirectoryManifest.Algorithm.StrongHashAlgorithmId, AllSources);
 
 				for (const FFileSyncTask& Task : AllFileTasks)
 				{
-					BlockRequestMap.AddFileBlocks(Task.OriginalSourceFilePath, Task.ResolvedSourceFilePath, *Task.SourceManifest);
+					BlockRequestMap.AddFileBlocks(Task.SourceId, Task.OriginalSourceFilePath, Task.ResolvedSourceFilePath, *Task.SourceManifest);
 				}
 
 				// Override loose file blocks with pack files
diff --git a/Engine/Source/Programs/Unsync/Private/UnsyncCore.h b/Engine/Source/Programs/Unsync/Private/UnsyncCore.h
--- a/Engine/Source/Programs/Unsync/Private/UnsyncCore.h
+++ b/Engine/Source/Programs/Unsync/Private/UnsyncCore.h
@@ -189,6 +189,8 @@ struct FFileSyncTask
 	uint64 NeedBytesFromBase   = 0;
 	uint64 TotalSizeBytes	   = 0;
 
+	uint32 SourceId = 0;
+
 	bool IsBaseValid() const { return !BaseFilePath.empty(); }
 };
 
diff --git a/Engine/Source/Programs/Unsync/Private/UnsyncHorde.cpp b/Engine/Source/Programs/Unsync/Private/UnsyncHorde.cpp
--- a/Engine/Source/Programs/Unsync/Private/UnsyncHorde.cpp
+++ b/Engine/Source/Programs/Unsync/Private/UnsyncHorde.cpp
@@ -21,6 +21,58 @@ FHordeProtocolImpl::FHordeProtocolImpl(const FRemoteDesc& InRemoteDesc, const FB
 
 FDownloadResult
 FHordeProtocolImpl::Download(const TArrayView<FNeedBlock> NeedBlocks, const FBlockDownloadCallback& CompletionCallback)
+{
+	// Requested blocks may come from different sources, in which case we have to split up the request by artifact root.
+
+	struct FBatch
+	{
+		std::vector<FNeedBlock> NeedBlocks;
+	};
+
+	std::vector<FBatch> Batches;
+	Batches.resize(RequestMap->GetSourceRoots().size());
+
+	for (const FNeedBlock Block : NeedBlocks)
+	{
+		const FBlockRequestMap::FBlockRequestEx* FoundRequest = RequestMap->FindRequest(Block.Hash);
+		if (!FoundRequest)
+		{
+			UNSYNC_ERROR(L"Could not find block request metadata");
+			return FDownloadError(EDownloadRetryMode::Abort);
+		}
+
+		if (FoundRequest->SourceId == ~0u)
+		{
+			UNSYNC_ERROR(L"Block request metadata does not contain a valid artifact source ID");
+			return FDownloadError(EDownloadRetryMode::Abort);
+		}
+
+		Batches[FoundRequest->SourceId].NeedBlocks.push_back(Block);
+	}
+
+	for (size_t BatchIndex = 0; BatchIndex < Batches.size(); ++BatchIndex)
+	{
+		const FBatch& Batch = Batches[BatchIndex];
+		if (Batch.NeedBlocks.empty())
+		{
+			continue;
+		}
+
+		const FPath& ArtifactPath	  = RequestMap->GetSourceRoots()[BatchIndex];
+		std::string	 ArtifactPathUtf8 = ConvertWideToUtf8(ArtifactPath.wstring());
+
+		FDownloadResult BatchDownloadResult = DownloadArtifactBlobs(ArtifactPathUtf8, MakeView(Batch.NeedBlocks), CompletionCallback);
+
+		if (BatchDownloadResult.IsError())
+		{
+			return BatchDownloadResult;
+		}
+	}
+
+	return ResultOk<FDownloadError>();
+}
+
+FDownloadResult FHordeProtocolImpl::DownloadArtifactBlobs(std::string_view ArtifactPath, const TArrayView<FNeedBlock> NeedBlocks, const FBlockDownloadCallback& CompletionCallback)
 {
 	if (NeedBlocks.Size() == 0)
 	{
@@ -29,8 +81,7 @@ FHordeProtocolImpl::Download(const TArrayView<FNeedBlock> NeedBlocks, const FBlo
 
 	std::string RequestJson = FormatBlockRequestJson(*RequestMap, NeedBlocks);
 
-	std::string RequestUrl = fmt::format("/{}/unsync-blobs?compress={}",
-										 ProxyPool.RemoteDesc.RequestPath,
+	std::string RequestUrl = fmt::format("/{}/unsync-blobs?compress={}", ArtifactPath,
 										 ProxyPool.RemoteDesc.bPreferCompression ? "true" : "false");
 
 	const EStrongHashAlgorithmID StrongHasher = RequestMap->GetStrongHasher();
diff --git a/Engine/Source/Programs/Unsync/Private/UnsyncHorde.h b/Engine/Source/Programs/Unsync/Private/UnsyncHorde.h
--- a/Engine/Source/Programs/Unsync/Private/UnsyncHorde.h
+++ b/Engine/Source/Programs/Unsync/Private/UnsyncHorde.h
@@ -28,6 +28,8 @@ struct FHordeProtocolImpl : FRemoteProtocolBase
 
 	FProxyPool& ProxyPool;
 	bool bValid = true;
+
+	FDownloadResult DownloadArtifactBlobs(std::string_view ArtifactPath, const TArrayView<FNeedBlock> NeedBlocks, const FBlockDownloadCallback& CompletionCallback);
 };
 
 
diff --git a/Engine/Source/Programs/Unsync/Private/UnsyncMain.cpp b/Engine/Source/Programs/Unsync/Private/UnsyncMain.cpp
--- a/Engine/Source/Programs/Unsync/Private/UnsyncMain.cpp
+++ b/Engine/Source/Programs/Unsync/Private/UnsyncMain.cpp
@@ -667,6 +667,25 @@ InnerMain(int Argc, char** Argv)
 	if (RemoteDesc.Protocol == EProtocolFlavor::Horde && Cli.got_subcommand(SubSync))
 	{
 		bFilesystemSource = false;
+
+		auto ResolveHordeArtifactPath = [](const std::string& PathUtf8)
+		{
+			TResult<FHordeArtifactQuery> Query = FHordeArtifactQuery::FromString(PathUtf8);
+			if (Query.IsError())
+			{
+				LogError(Query.GetError(), L"Could not parse sync source path");
+				return std::string();
+			}
+
+			if (Query->Id.empty())
+			{
+				UNSYNC_ERROR(L"Could not parse sync source path. Artifact ID is expected, i.e. '#123456abcdef'.");
+				return std::string();
+			}
+
+			return std::string("api/v2/artifacts/") + Query->Id;
+		};
+
 		if (RemoteDesc.RequestPath.empty())
 		{
 			TResult<FHordeArtifactQuery> Query = FHordeArtifactQuery::FromString(SourceFilenameUtf8);
@@ -682,9 +701,14 @@ InnerMain(int Argc, char** Argv)
 				return 1;
 			}
 
-			SourceFilenameUtf8 = "api/v2/artifacts/" + Query->Id;
+			SourceFilenameUtf8	   = ResolveHordeArtifactPath(SourceFilenameUtf8);
 			RemoteDesc.RequestPath = SourceFilenameUtf8;
 		}
+
+		for (std::string& OverlayPath : OverlayArrayUtf8)
+		{
+			OverlayPath = ResolveHordeArtifactPath(OverlayPath);
+		}
 	}
 
 	if (bNoCompression)
@@ -1080,7 +1104,14 @@ InnerMain(int Argc, char** Argv)
 
 		for (const std::string& Entry : OverlayArrayUtf8)
 		{
-			SyncOptions.Overlays.push_back(NormalizeFilenameUtf8(Entry));
+			if (bFilesystemSource)
+			{
+				SyncOptions.Overlays.push_back(NormalizeFilenameUtf8(Entry));
+			}
+			else
+			{
+				SyncOptions.Overlays.push_back(FPath(Entry));
+			}
 		}
 
 		return CmdSync(SyncOptions);
diff --git a/Engine/Source/Programs/Unsync/Private/UnsyncProxy.cpp b/Engine/Source/Programs/Unsync/Private/UnsyncProxy.cpp
--- a/Engine/Source/Programs/Unsync/Private/UnsyncProxy.cpp
+++ b/Engine/Source/Programs/Unsync/Private/UnsyncProxy.cpp
@@ -321,7 +321,7 @@ FUnsyncProtocolImpl::Download(const TArrayView<FNeedBlock> NeedBlocks, const FBl
 	std::vector<const std::string*> FileListUtf8;
 	for (const FHash128& It : UniqueFileNamesMd5)
 	{
-		const std::string* Name = RequestMap->FindFile(It);
+		const std::string* Name = RequestMap->FindSourceFile(It);
 		if (Name)
 		{
 			FileListUtf8.push_back(Name);
@@ -519,7 +519,7 @@ FormatBlockRequestJson(const FBlockRequestMap& RequestMap, const TArrayView<FNee
 				Output += "]},\n";	// close blocks arary and file object
 			}
 
-			const std::string* FilenameUtf8 = RequestMap.FindFile(Request.FilenameMd5);
+			const std::string* FilenameUtf8 = RequestMap.FindSourceFile(Request.FilenameMd5);
 			UNSYNC_ASSERTF(FilenameUtf8, L"Could not find file in the block request map");
 
 			// Start file object and blocks array
@@ -1145,9 +1145,9 @@ FBlockRequestMap::AddFile(const FPath& OriginalFilePath, const FPath& ResolvedFi
 	auto FindResult = HashToFile.find(OriginalNameHash);
 	if (FindResult == HashToFile.end())
 	{
-		HashToFile[OriginalNameHash] = uint32(FileListUtf8.size());
-		HashToFile[ResolvedNameHash] = uint32(FileListUtf8.size());
-		FileListUtf8.push_back(OriginalFilePathUtf8);
+		HashToFile[OriginalNameHash] = uint32(SourceFileListUtf8.size());
+		HashToFile[ResolvedNameHash] = uint32(SourceFileListUtf8.size());
+		SourceFileListUtf8.push_back(OriginalFilePathUtf8);
 	}
 
 	return OriginalNameHash;
@@ -1164,7 +1164,7 @@ FBlockRequestMap::AddPackBlocks(const FPath&					  OriginalFilePath,
 
 	for (const FPackIndexEntry& Block : PackManifest)
 	{
-		FBlockRequest Request;
+		FBlockRequestEx Request;
 		Request.FilenameMd5				 = FileId;
 		Request.BlockHash				 = Block.BlockHash;
 		Request.Offset					 = Block.PackBlockOffset;
@@ -1174,7 +1174,7 @@ FBlockRequestMap::AddPackBlocks(const FPath&					  OriginalFilePath,
 }
 
 void
-FBlockRequestMap::AddFileBlocks(const FPath& OriginalFilePath, const FPath& ResolvedFilePath, const FFileManifest& FileManifest)
+FBlockRequestMap::AddFileBlocks(uint32 SourceId, const FPath& OriginalFilePath, const FPath& ResolvedFilePath, const FFileManifest& FileManifest)
 {
 	UNSYNC_ASSERTF(StrongHasher != EStrongHashAlgorithmID::Invalid, L"Request map is not initialized");
 
@@ -1182,11 +1182,12 @@ FBlockRequestMap::AddFileBlocks(const FPath& OriginalFilePath, const FPath& Reso
 
 	for (const FGenericBlock& Block : FileManifest.Blocks)
 	{
-		FBlockRequest Request;
+		FBlockRequestEx Request;
 		Request.FilenameMd5				 = FileId;
 		Request.BlockHash				 = Block.HashStrong.ToHash128();  // #wip-widehash
 		Request.Offset					 = Block.Offset;
 		Request.Size					 = Block.Size;
+		Request.SourceId				 = SourceId;
 		BlockRequests[Request.BlockHash] = Request;
 
 		if (!FileManifest.MacroBlocks.empty())
@@ -1225,7 +1226,7 @@ FBlockRequestMap::AddFileBlocks(const FPath& OriginalFilePath, const FPath& Reso
 	}
 }
 
-const FBlockRequest*
+const FBlockRequestMap::FBlockRequestEx*
 FBlockRequestMap::FindRequest(const FGenericHash& BlockHash) const
 {
 	FHash128 BlockHash128 = BlockHash.ToHash128();
@@ -1242,16 +1243,16 @@ FBlockRequestMap::FindRequest(const FGenericHash& BlockHash) const
 }
 
 const std::string*
-FBlockRequestMap::FindFile(const FHash128& Hash) const
+FBlockRequestMap::FindSourceFile(const FHash128& NameHashMd5) const
 {
-	auto It = HashToFile.find(Hash);
+	auto It = HashToFile.find(NameHashMd5);
 	if (It == HashToFile.end())
 	{
 		return nullptr;
 	}
 	else
 	{
-		return &FileListUtf8[It->second];
+		return &SourceFileListUtf8[It->second];
 	}
 }
 
@@ -1422,13 +1423,6 @@ FProxyPool::IsValid() const
 	return bValid;
 }
 
-void
-FProxyPool::BuildFileBlockRequests(const FPath& OriginalFilePath, const FPath& ResolvedFilePath, const FFileManifest& FileManifest)
-{
-	std::lock_guard<std::mutex> LockGuard(Mutex);
-	RequestMap.AddFileBlocks(OriginalFilePath, ResolvedFilePath, FileManifest);
-}
-
 void
 FProxyPool::SendTelemetryEvent(const FTelemetryEventSyncComplete& Event)
 {
@@ -1438,13 +1432,6 @@ FProxyPool::SendTelemetryEvent(const FTelemetryEventSyncComplete& Event)
 	}
 }
 
-void
-FProxyPool::InitRequestMap(EStrongHashAlgorithmID InStrongHasher)
-{
-	std::lock_guard<std::mutex> LockGuard(Mutex);
-	RequestMap.Init(InStrongHasher);
-}
-
 void
 FProxyPool::SetRequestMap(FBlockRequestMap&& InRequestMap)
 {
diff --git a/Engine/Source/Programs/Unsync/Private/UnsyncProxy.h b/Engine/Source/Programs/Unsync/Private/UnsyncProxy.h
--- a/Engine/Source/Programs/Unsync/Private/UnsyncProxy.h
+++ b/Engine/Source/Programs/Unsync/Private/UnsyncProxy.h
@@ -66,30 +66,38 @@ struct FMacroBlockRequest
 class FBlockRequestMap
 {
 public:
-	void Init(EStrongHashAlgorithmID InStrongHasher)
+	void Init(EStrongHashAlgorithmID InStrongHasher, const std::vector<FPath>& InSourceRoots)
 	{
 		UNSYNC_ASSERTF(StrongHasher == EStrongHashAlgorithmID::Invalid, L"Request map is already initialized");
 		StrongHasher = InStrongHasher;
+		SourceRoots	 = InSourceRoots;
 	}
 
-	void AddFileBlocks(const FPath& OriginalFilePath, const FPath& ResolvedFilePath, const FFileManifest& Manifest);
+	void AddFileBlocks(uint32 SourceId, const FPath& OriginalFilePath, const FPath& ResolvedFilePath, const FFileManifest& Manifest);
 	void AddPackBlocks(const FPath& OriginalFilePath, const FPath& ResolvedFilePath, const TArrayView<FPackIndexEntry> PackManifest);
 
-	const std::vector<std::string>& GetFileList() const { return FileListUtf8; }
-	const FBlockRequest*			FindRequest(const FGenericHash& BlockHash) const;
-	const std::string*				FindFile(const FHash128& Hash) const;
+	struct FBlockRequestEx : FBlockRequest
+	{
+		uint32 SourceId = ~0u;
+	};
+
+	const std::vector<std::string>& GetSourceFileList() const { return SourceFileListUtf8; }
+	const FBlockRequestEx*			FindRequest(const FGenericHash& BlockHash) const;
+	const std::string*				FindSourceFile(const FHash128& NameHashMd5) const;
 	EStrongHashAlgorithmID			GetStrongHasher() const { return StrongHasher; }
 	FMacroBlockRequest				GetMacroBlockRequest(const FGenericHash& BlockHash) const;
+	const std::vector<FPath>&		GetSourceRoots() const { return SourceRoots; }
 
 private:
 
 	FHash128 AddFile(const FPath& OriginalFilePath, const FPath& ResolvedFilePath);
 
 	EStrongHashAlgorithmID							 StrongHasher = EStrongHashAlgorithmID::Invalid;
-	std::vector<std::string>						 FileListUtf8;
+	std::vector<std::string>						 SourceFileListUtf8;
 	std::unordered_map<FHash128, uint32>			 HashToFile;
-	std::unordered_map<FHash128, FBlockRequest>		 BlockRequests;
+	std::unordered_map<FHash128, FBlockRequestEx>	 BlockRequests;
 	std::unordered_map<FHash128, FMacroBlockRequest> MacroBlockRequests;
+	std::vector<FPath>								 SourceRoots;
 };
 
 struct FRemoteProtocolFeatures
@@ -185,9 +193,7 @@ public:
 	const FRemoteDesc RemoteDesc;
 	const FAuthDesc* AuthDesc = nullptr; // optional reference to externally-owned auth parameters
 
-	void InitRequestMap(EStrongHashAlgorithmID InStrongHasher);
 	void SetRequestMap(FBlockRequestMap&& InRequestMap);
-	void BuildFileBlockRequests(const FPath& OriginalFilePath, const FPath& ResolvedFilePath, const FFileManifest& FileManifest);
 
 	const FRemoteProtocolFeatures& GetFeatures() const { return Features; }
 	const std::string& GetSessionId() const { return SessionId; }
diff --git a/Engine/Source/Programs/Unsync/Private/UnsyncSerialization.cpp b/Engine/Source/Programs/Unsync/Private/UnsyncSerialization.cpp
--- a/Engine/Source/Programs/Unsync/Private/UnsyncSerialization.cpp
+++ b/Engine/Source/Programs/Unsync/Private/UnsyncSerialization.cpp
@@ -824,7 +824,7 @@ SaveDirectoryManifest(const FDirectoryManifest& Manifest, FVectorStreamOut& Stre
 }
 
 bool
-SaveDirectoryManifest(const FDirectoryManifest& Manifest, const FPath& Filename)
+SaveDirectoryManifest(const FDirectoryManifest& Manifest, const FPath& Filename, bool bAllowInDryRun)
 {
 	FBuffer			 OutputBuffer;
 	FVectorStreamOut OutputStream(OutputBuffer);
@@ -834,7 +834,13 @@ SaveDirectoryManifest(const FDirectoryManifest& Manifest, const FPath& Filename)
 	bool bSerialized = SaveDirectoryManifest(Manifest, OutputStream);
 	UNSYNC_ASSERT(bSerialized);
 
-	FNativeFile OutputFile(Filename, EFileMode::CreateWriteOnly, OutputBuffer.Size());
+	EFileMode FileMode = EFileMode ::CreateWriteOnly;
+	if (bAllowInDryRun)
+	{
+		FileMode = FileMode | EFileMode::IgnoreDryRun;
+	}
+
+	FNativeFile OutputFile(Filename, FileMode, OutputBuffer.Size());
 	if (OutputFile.IsValid())
 	{
 		uint64 WroteBytes = OutputFile.Write(OutputBuffer.Data(), 0, OutputBuffer.Size());
diff --git a/Engine/Source/Programs/Unsync/Private/UnsyncSerialization.h b/Engine/Source/Programs/Unsync/Private/UnsyncSerialization.h
--- a/Engine/Source/Programs/Unsync/Private/UnsyncSerialization.h
+++ b/Engine/Source/Programs/Unsync/Private/UnsyncSerialization.h
@@ -16,6 +16,6 @@ bool LoadDirectoryManifest(FDirectoryManifest& OutManifest, const FPath& Root, F
 bool LoadDirectoryManifest(FDirectoryManifest& OutManifest, const FPath& Root, const FPath& Filename);
 
 bool SaveDirectoryManifest(const FDirectoryManifest& Manifest, FVectorStreamOut& Stream);
-bool SaveDirectoryManifest(const FDirectoryManifest& Manifest, const FPath& Filename);
+bool SaveDirectoryManifest(const FDirectoryManifest& Manifest, const FPath& Filename, bool bAllowInDryRun = false);
 
 }  // namespace unsync
diff --git a/Engine/Source/Programs/Unsync/Private/UnsyncVersion.cpp b/Engine/Source/Programs/Unsync/Private/UnsyncVersion.cpp
--- a/Engine/Source/Programs/Unsync/Private/UnsyncVersion.cpp
+++ b/Engine/Source/Programs/Unsync/Private/UnsyncVersion.cpp
@@ -2,6 +2,8 @@
 
 #include "UnsyncVersion.h"
 
+#define UNSYNC_VERSION_STR "1.0.85"
+
 namespace unsync {
 
 const std::string&
@@ -35,4 +37,4 @@ GetVersionString()
 	return Result;
 }
 
-}  // namespace unsync
\ No newline at end of file
+}  // namespace unsync
diff --git a/Engine/Source/Programs/Unsync/Private/UnsyncVersion.h b/Engine/Source/Programs/Unsync/Private/UnsyncVersion.h
--- a/Engine/Source/Programs/Unsync/Private/UnsyncVersion.h
+++ b/Engine/Source/Programs/Unsync/Private/UnsyncVersion.h
@@ -4,8 +4,6 @@
 
 #include "UnsyncCommon.h"
 
-#define UNSYNC_VERSION_STR "1.0.84"
-
 namespace unsync {
 
 const std::string&
