unsync - Fix GetRelativePath() for network paths

diff --git a/Engine/Source/Programs/Unsync/Private/UnsyncCore.cpp b/Engine/Source/Programs/Unsync/Private/UnsyncCore.cpp
--- a/Engine/Source/Programs/Unsync/Private/UnsyncCore.cpp
+++ b/Engine/Source/Programs/Unsync/Private/UnsyncCore.cpp
@@ -1251,7 +1251,7 @@ DeleteOldFilesInDirectory(const FPath& Path, uint32 MaxFilesToKeep)
 	{
 		const FEntry& Oldest = Entries.back();
 
-		std::wstring PathStr = FPath(RemoveExtendedPathPrefix(Oldest.Path)).wstring();
+		std::wstring PathStr = RemoveExtendedPathPrefix(Oldest.Path).wstring();
 
 		if (GDryRun)
 		{
diff --git a/Engine/Source/Programs/Unsync/Private/UnsyncFile.cpp b/Engine/Source/Programs/Unsync/Private/UnsyncFile.cpp
--- a/Engine/Source/Programs/Unsync/Private/UnsyncFile.cpp
+++ b/Engine/Source/Programs/Unsync/Private/UnsyncFile.cpp
@@ -63,10 +63,16 @@ MakeExtendedAbsolutePath(const FPath& InAbsolutePath)
 #endif // UNSYNC_PLATFORM_WINDOWS
 }
 
-FPathStringView
-RemoveExtendedPathPrefix(const FPath& InPath)
+// Removes \\, \\?\UNC\, \\.\UNC\, \\.\ or \\?\ prefix from a path.
+//	\\server\foo\bar -> server\foo\bar
+//	\\?\d:\foo\bar -> d:\foo\bar
+//	d:\foo\bar -> d:\foo\bar
+// Returns original path on non-Windows.
+static inline FPathStringView
+RemoveUNCPrefix(const FPath& InPath)
 {
 	FPathStringView InPathString = InPath.native();
+
 #if UNSYNC_PLATFORM_WINDOWS
 	if (InPathString.starts_with(L"\\\\?\\UNC\\"))
 	{
@@ -76,6 +82,35 @@ RemoveExtendedPathPrefix(const FPath& InPath)
 	{
 		return InPathString.substr(4);
 	}
+	else if (InPathString.starts_with(L"\\\\"))
+	{
+		return InPathString.substr(2);
+	}
+	else
+#endif
+	{
+		return InPathString;
+	}
+}
+
+FPath
+RemoveExtendedPathPrefix(const FPath& InPath)
+{
+	FPathStringView InPathString = InPath.native();
+#if UNSYNC_PLATFORM_WINDOWS
+	if (InPathString.starts_with(L"\\\\?\\UNC\\"))
+	{
+		FPathStringView Remainder = InPathString.substr(8);
+		std::wstring Result;
+		Result.reserve(Remainder.length() + 2);
+		Result += L"\\\\";
+		Result += Remainder;
+		return FPath(Result);
+	}
+	else if (InPathString.starts_with(L"\\\\?\\"))
+	{
+		return InPathString.substr(4);
+	}
 	else
 	{
 		return InPathString;
@@ -119,19 +154,21 @@ FPathStringView
 GetRelativePathView(const FPath& Path, const FPath& Base)
 {
 	// Try a trivial case first, without touching the filesystem
-	FPathStringView PathView = RemoveExtendedPathPrefix(Path);
-	FPathStringView BaseView = RemoveExtendedPathPrefix(Base);
+	FPathStringView PathView = RemoveUNCPrefix(Path);
+	FPathStringView BaseView = RemoveUNCPrefix(Base);
 
-	FPathStringView PathViewRemainder = PathView.substr(BaseView.length());
-
-	if (PathView.starts_with(BaseView) && PathViewRemainder.starts_with(FPath::preferred_separator))
+	if (PathView.starts_with(BaseView))
 	{
-		FPathStringView RelativePath = PathView.substr(BaseView.length());
-		while (RelativePath.starts_with(FPath::preferred_separator))
+		FPathStringView PathViewRemainder = PathView.substr(BaseView.length());
+		if (PathViewRemainder.starts_with(FPath::preferred_separator))
 		{
-			RelativePath = RelativePath.substr(1);
+			FPathStringView RelativePath = PathView.substr(BaseView.length());
+			while (RelativePath.starts_with(FPath::preferred_separator))
+			{
+				RelativePath = RelativePath.substr(1);
+			}
+			return RelativePath;
 		}
-		return RelativePath;
 	}
 
 	return {};
@@ -1587,4 +1624,92 @@ TestFileAttrib()
 	UNSYNC_ASSERT(bFileDeleted);
 }
 
+void
+TestPathUtil()
+{
+#if UNSYNC_PLATFORM_WINDOWS
+
+	UNSYNC_LOG(L"TestPathUtil()");
+	UNSYNC_LOG_INDENT;
+
+	// Test path manipulation helpers
+
+	{
+		FPath Simple   = FPath("\\\\?\\UNC\\server\\subdir\\a\\b\\c");
+		FPath Extended = MakeExtendedAbsolutePath(Simple);
+		UNSYNC_ASSERT(Simple == Extended);
+	}
+
+	{
+		FPath Simple   = FPath("\\\\?\\d:\\local\\subdir\\a\\b\\c");
+		FPath Extended = MakeExtendedAbsolutePath(Simple);
+		UNSYNC_ASSERT(Simple == Extended);
+	}
+
+	{
+		FPath Simple   = FPath("d:\\local\\subdir\\a\\b\\c");
+		FPath Extended = MakeExtendedAbsolutePath(Simple);
+		FPath Stripped = RemoveExtendedPathPrefix(Extended);
+		UNSYNC_ASSERT(Stripped == Simple);
+	}
+
+	{
+		FPath Simple   = FPath("\\\\server\\local\\subdir\\a\\b\\c");
+		FPath Extended = MakeExtendedAbsolutePath(Simple);
+		FPath Stripped = RemoveExtendedPathPrefix(Extended);
+		UNSYNC_ASSERT(Stripped == Simple);
+	}
+
+	{
+		FPath Base	   = FPath("d:\\local\\subdir");
+		FPath Full	   = FPath("d:\\local\\subdir\\a\\b\\c");
+		FPath Relative = GetRelativePath(Full, Base);
+		UNSYNC_ASSERT(Relative == FPath("a\\b\\c"));
+	}
+
+	{
+		FPath Base	   = FPath("\\\\server\\subdir");
+		FPath Full	   = FPath("\\\\server\\subdir\\a\\b\\c");
+		FPath Relative = GetRelativePath(Full, Base);
+		UNSYNC_ASSERT(Relative == FPath("a\\b\\c"));
+	}
+
+	{
+		FPath Base	   = FPath("\\\\?\\d:\\local\\subdir");
+		FPath Full	   = FPath("\\\\?\\d:\\local\\subdir\\a\\b\\c");
+		FPath Relative = GetRelativePath(Full, Base);
+		UNSYNC_ASSERT(Relative == FPath("a\\b\\c"));
+	}
+
+	{
+		FPath Base	   = FPath("\\\\?\\d:\\local\\subdir");
+		FPath Full	   = FPath("d:\\local\\subdir\\a\\b\\c");
+		FPath Relative = GetRelativePath(Full, Base);
+		UNSYNC_ASSERT(Relative == FPath("a\\b\\c"));
+	}
+
+	{
+		FPath Base	   = FPath("d:\\local\\subdir");
+		FPath Full	   = FPath("\\\\?\\d:\\local\\subdir\\a\\b\\c");
+		FPath Relative = GetRelativePath(Full, Base);
+		UNSYNC_ASSERT(Relative == FPath("a\\b\\c"));
+	}
+
+	{
+		FPath Base	   = FPath("d:\\local\\subdir");
+		FPath Full	   = FPath("\\\\?\\e:\\local\\subdir\\a\\b\\c");
+		FPath Relative = GetRelativePath(Full, Base);
+		UNSYNC_ASSERT(Relative.empty());
+	}
+
+	{
+		FPath Base	   = FPath("d:\\local\\subdir");
+		FPath Full	   = FPath("d:\\local\\a\\b\\c");
+		FPath Relative = GetRelativePath(Full, Base);
+		UNSYNC_ASSERT(Relative.empty());
+	}
+
+#endif // UNSYNC_PLATFORM_WINDOWS
+}
+
 }  // namespace unsync
diff --git a/Engine/Source/Programs/Unsync/Private/UnsyncFile.h b/Engine/Source/Programs/Unsync/Private/UnsyncFile.h
--- a/Engine/Source/Programs/Unsync/Private/UnsyncFile.h
+++ b/Engine/Source/Programs/Unsync/Private/UnsyncFile.h
@@ -527,6 +527,6 @@ FPath MakeExtendedAbsolutePath(const FPath& InAbsolutePath);
 
 // Removes `\\?\` or `\\?\UNC\` prefix from a given path.
 // Returns original path on non-Windows.
-FPathStringView RemoveExtendedPathPrefix(const FPath& InPath);
+FPath RemoveExtendedPathPrefix(const FPath& InPath);
 
 }  // namespace unsync
diff --git a/Engine/Source/Programs/Unsync/Private/UnsyncTest.cpp b/Engine/Source/Programs/Unsync/Private/UnsyncTest.cpp
--- a/Engine/Source/Programs/Unsync/Private/UnsyncTest.cpp
+++ b/Engine/Source/Programs/Unsync/Private/UnsyncTest.cpp
@@ -947,6 +947,12 @@ RunTests(const std::string& Preset)
 		TestFileAttrib();
 	}
 
+	if (Preset == "pathutil" || Preset == "all")
+	{
+		extern void TestPathUtil();
+		TestPathUtil();
+	}
+
 	if (Preset == "horde_manifest_decode" || Preset == "all")
 	{
 		extern void TestHordeManifestDecode();
diff --git a/Engine/Source/Programs/Unsync/Private/UnsyncVersion.cpp b/Engine/Source/Programs/Unsync/Private/UnsyncVersion.cpp
--- a/Engine/Source/Programs/Unsync/Private/UnsyncVersion.cpp
+++ b/Engine/Source/Programs/Unsync/Private/UnsyncVersion.cpp
@@ -2,7 +2,7 @@
 
 #include "UnsyncVersion.h"
 
-#define UNSYNC_VERSION_STR "1.0.85"
+#define UNSYNC_VERSION_STR "1.0.86"
 
 namespace unsync {
 
