unsync - Suppress deprecation warnings from standard headers related to unicode conversion utilities when compiling with Clang

diff --git a/Engine/Source/Programs/Unsync/Private/UnsyncCommon.h b/Engine/Source/Programs/Unsync/Private/UnsyncCommon.h
--- a/Engine/Source/Programs/Unsync/Private/UnsyncCommon.h
+++ b/Engine/Source/Programs/Unsync/Private/UnsyncCommon.h
@@ -2,10 +2,43 @@
 
 #pragma once
 
+#if !defined(UNSYNC_COMPILER_MSVC) && defined(_MSC_VER) && !defined(__clang__)
+#	define UNSYNC_COMPILER_MSVC 1
+#elif !defined(UNSYNC_COMPILER_MSVC)
+#	define UNSYNC_COMPILER_MSVC 0
+#endif
+
+#if UNSYNC_COMPILER_MSVC
+#	define UNSYNC_BREAK __debugbreak
+#else
+#	define UNSYNC_BREAK __builtin_trap
+#endif
+
+#if UNSYNC_COMPILER_MSVC
+#	pragma warning(disable : 4100)	 // unreferenced formal parameter
+#	define UNSYNC_ATTRIB_FORCEINLINE [[msvc::forceinline]]
+#	define UNSYNC_THIRD_PARTY_INCLUDES_START \
+		__pragma(warning(push)) \
+		__pragma(warning(disable : 4458)) \
+		__pragma(warning(disable : 4668)) /* undefined macro, replaced with 0 */ \
+		__pragma(warning(disable : 4996)) /* deprecated */
+#	define UNSYNC_THIRD_PARTY_INCLUDES_END __pragma(warning(pop))
+#else
+#	define UNSYNC_ATTRIB_FORCEINLINE __attribute__((always_inline))
+#	define UNSYNC_THIRD_PARTY_INCLUDES_START \
+		_Pragma("GCC diagnostic push") \
+		_Pragma("GCC diagnostic ignored \"-Wdeprecated-declarations\"") \
+		_Pragma("GCC diagnostic ignored \"-Wunknown-pragmas\"") \
+		_Pragma("GCC diagnostic ignored \"-Wshadow\"")
+#	define UNSYNC_THIRD_PARTY_INCLUDES_END _Pragma("GCC diagnostic pop")
+#endif
+
+UNSYNC_THIRD_PARTY_INCLUDES_START
 #include <stdint.h>
 #include <cstring>
 #include <string_view>
 #include <filesystem>
+UNSYNC_THIRD_PARTY_INCLUDES_END
 
 namespace unsync {
 
@@ -46,36 +79,6 @@
 
 #define UNSYNC_CLOBBER(Object) ClobberT(Object)
 
-#if !defined(UNSYNC_COMPILER_MSVC) && defined(_MSC_VER) && !defined(__clang__)
-#	define UNSYNC_COMPILER_MSVC 1
-#elif !defined(UNSYNC_COMPILER_MSVC)
-#	define UNSYNC_COMPILER_MSVC 0
-#endif
-
-#if UNSYNC_COMPILER_MSVC
-#	define UNSYNC_BREAK __debugbreak
-#else
-#	define UNSYNC_BREAK __builtin_trap
-#endif
-
-#if UNSYNC_COMPILER_MSVC
-#	pragma warning(disable : 4100)	 // unreferenced formal parameter
-#	define UNSYNC_ATTRIB_FORCEINLINE [[msvc::forceinline]]
-#	define UNSYNC_THIRD_PARTY_INCLUDES_START \
-		__pragma(warning(push)) \
-		__pragma(warning(disable : 4458)) \
-		__pragma(warning(disable : 4668)) /* undefined macro, replaced with 0 */ \
-		__pragma(warning(disable : 4996)) /* deprecated */
-#	define UNSYNC_THIRD_PARTY_INCLUDES_END __pragma(warning(pop))
-#else
-#	define UNSYNC_ATTRIB_FORCEINLINE __attribute__((always_inline))
-#	define UNSYNC_THIRD_PARTY_INCLUDES_START \
-		_Pragma("GCC diagnostic push") \
-		_Pragma("GCC diagnostic ignored \"-Wunknown-pragmas\"") \
-		_Pragma("GCC diagnostic ignored \"-Wshadow\"")
-#	define UNSYNC_THIRD_PARTY_INCLUDES_END _Pragma("GCC diagnostic pop")
-#endif
-
 #define UNSYNC_ENUM_CLASS_FLAGS(T, S)                      \
 	inline constexpr T operator^(T a, T b)                 \
 	{                                                      \
diff --git a/Engine/Source/Programs/Unsync/Private/UnsyncLog.cpp b/Engine/Source/Programs/Unsync/Private/UnsyncLog.cpp
--- a/Engine/Source/Programs/Unsync/Private/UnsyncLog.cpp
+++ b/Engine/Source/Programs/Unsync/Private/UnsyncLog.cpp
@@ -15,8 +15,10 @@
 #include <vector>
 #include <chrono>
 
+UNSYNC_THIRD_PARTY_INCLUDES_START
 #include <fmt/format.h>
 #include <fmt/chrono.h>
+UNSYNC_THIRD_PARTY_INCLUDES_END
 
 #if UNSYNC_PLATFORM_WINDOWS
 #	include <Dbghelp.h>
