unsync - Keep backup copies of 10 last log files

* Also add log file writing on Linux/Mac

diff --git a/Engine/Source/Programs/Unsync/Private/UnsyncCore.cpp b/Engine/Source/Programs/Unsync/Private/UnsyncCore.cpp
--- a/Engine/Source/Programs/Unsync/Private/UnsyncCore.cpp
+++ b/Engine/Source/Programs/Unsync/Private/UnsyncCore.cpp
@@ -1221,53 +1221,6 @@ struct FFileSyncTaskBatch
 	}
 };
 
-static void
-DeleteOldFilesInDirectory(const FPath& Path, uint32 MaxFilesToKeep)
-{
-	struct FEntry
-	{
-		FPath  Path;
-		uint64 Mtime;
-	};
-
-	std::vector<FEntry> Entries;
-
-	FPath ExtendedPath = MakeExtendedAbsolutePath(Path);
-	for (const std::filesystem::directory_entry& It : std::filesystem::directory_iterator(ExtendedPath))
-	{
-		if (It.is_regular_file())
-		{
-			FEntry Entry;
-			Entry.Mtime = ToWindowsFileTime(It.last_write_time());
-			Entry.Path	= It.path();
-			Entries.push_back(Entry);
-		}
-	}
-
-	// reverse sort
-	std::sort(Entries.begin(), Entries.end(), [](const FEntry& A, const FEntry& B) { return A.Mtime > B.Mtime; });
-
-	while (Entries.size() > MaxFilesToKeep)
-	{
-		const FEntry& Oldest = Entries.back();
-
-		std::wstring PathStr = RemoveExtendedPathPrefix(Oldest.Path).wstring();
-
-		if (GDryRun)
-		{
-			UNSYNC_VERBOSE(L"Deleting '%ls'(skipped due to dry run mode)", PathStr.c_str());
-		}
-		else
-		{
-			UNSYNC_VERBOSE(L"Deleting '%ls'", PathStr.c_str());
-			std::error_code ErrorCode = {};
-			FileRemove(Oldest.Path, ErrorCode);
-		}
-
-		Entries.pop_back();
-	}
-}
-
 bool  // TODO: return a TResult
 SyncDirectory(const FSyncDirectoryOptions& SyncOptions)
 {
diff --git a/Engine/Source/Programs/Unsync/Private/UnsyncFile.cpp b/Engine/Source/Programs/Unsync/Private/UnsyncFile.cpp
--- a/Engine/Source/Programs/Unsync/Private/UnsyncFile.cpp
+++ b/Engine/Source/Programs/Unsync/Private/UnsyncFile.cpp
@@ -1712,4 +1712,77 @@ TestPathUtil()
 #endif // UNSYNC_PLATFORM_WINDOWS
 }
 
+void
+DeleteOldFilesInDirectory(const FPath& Path, uint32 MaxFilesToKeep, bool bAllowInDryRun, const FPathFilterCallback& Filter)
+{
+	struct FEntry
+	{
+		FPath  Path;
+		uint64 Mtime;
+	};
+
+	std::vector<FEntry> Entries;
+
+	FPath ExtendedPath = MakeExtendedAbsolutePath(Path);
+	for (const std::filesystem::directory_entry& It : std::filesystem::directory_iterator(ExtendedPath))
+	{
+		if (Filter && !Filter(It.path()))
+		{
+			continue;
+		}
+
+		if (It.is_regular_file())
+		{
+			FEntry Entry;
+			Entry.Mtime = ToWindowsFileTime(It.last_write_time());
+			Entry.Path	= It.path();
+			Entries.push_back(Entry);
+		}
+	}
+
+	// reverse sort
+	std::sort(Entries.begin(), Entries.end(), [](const FEntry& A, const FEntry& B) { return A.Mtime > B.Mtime; });
+
+	while (Entries.size() > MaxFilesToKeep)
+	{
+		const FEntry& Oldest = Entries.back();
+
+		std::wstring PathStr = RemoveExtendedPathPrefix(Oldest.Path).wstring();
+
+		if (GDryRun && !bAllowInDryRun)
+		{
+			UNSYNC_VERBOSE(L"Deleting '%ls'(skipped due to dry run mode)", PathStr.c_str());
+		}
+		else
+		{
+			UNSYNC_VERBOSE(L"Deleting '%ls'", PathStr.c_str());
+			std::error_code ErrorCode = {};
+			FileRemove(Oldest.Path, ErrorCode);
+		}
+
+		Entries.pop_back();
+	}
+}
+
+std::chrono::system_clock::time_point
+SystemTimeFromFileTime(std::filesystem::file_time_type FileTime)
+{
+#if 0 // TODO: use clock_cast() when it's supported in all target compilers
+
+	return std::chrono::clock_cast<std::chrono::system_clock>(FileTime);
+
+#else
+
+#if UNSYNC_PLATFORM_WINDOWS
+	std::chrono::seconds Offset(SECONDS_BETWEEN_WINDOWS_AND_UNIX);
+#else // UNSYNC_PLATFORM_WINDOWS
+	std::chrono::seconds Offset(0);
+#endif // UNSYNC_PLATFORM_WINDOWS
+
+	std::chrono::duration SinceEpoch = FileTime.time_since_epoch() - Offset;
+	return std::chrono::system_clock::time_point(std::chrono::duration_cast<std::chrono::system_clock::duration>(SinceEpoch));
+
+#endif // 0
+}
+
 }  // namespace unsync
diff --git a/Engine/Source/Programs/Unsync/Private/UnsyncFile.h b/Engine/Source/Programs/Unsync/Private/UnsyncFile.h
--- a/Engine/Source/Programs/Unsync/Private/UnsyncFile.h
+++ b/Engine/Source/Programs/Unsync/Private/UnsyncFile.h
@@ -516,6 +516,7 @@ std::filesystem::directory_iterator DirectoryScan(const FPath& Path);
 
 uint64 ToWindowsFileTime(const std::filesystem::file_time_type& T);
 std::filesystem::file_time_type FromWindowsFileTime(uint64 Ticks);
+std::chrono::system_clock::time_point SystemTimeFromFileTime(std::filesystem::file_time_type FileTime);
 
 struct FSyncFilter;
 FFileAttributeCache CreateFileAttributeCache(const FPath& Root, const FSyncFilter* SyncFilter = nullptr);
@@ -529,4 +530,7 @@ FPath MakeExtendedAbsolutePath(const FPath& InAbsolutePath);
 // Returns original path on non-Windows.
 FPath RemoveExtendedPathPrefix(const FPath& InPath);
 
+using FPathFilterCallback = std::function<bool(const FPath& Path)>;
+void DeleteOldFilesInDirectory(const FPath& Path, uint32 MaxFilesToKeep, bool bAllowInDryRun = false, const FPathFilterCallback& Filter = {});
+
 }  // namespace unsync
diff --git a/Engine/Source/Programs/Unsync/Private/UnsyncLog.cpp b/Engine/Source/Programs/Unsync/Private/UnsyncLog.cpp
--- a/Engine/Source/Programs/Unsync/Private/UnsyncLog.cpp
+++ b/Engine/Source/Programs/Unsync/Private/UnsyncLog.cpp
@@ -28,6 +28,11 @@ UNSYNC_THIRD_PARTY_INCLUDES_END
 
 namespace unsync {
 
+static constexpr size_t TimestampStringSize = sizeof("YYYY-MM-DDTHH:MM:SS.sssZ");
+using FTimestampString = fmt::basic_memory_buffer<char, TimestampStringSize>;
+static FTimestampString FormatTimestamp(std::chrono::system_clock::time_point Timestamp);
+static FTimestampString FormatWindowsFileTime(uint64 Ticks);
+
 #if UNSYNC_PLATFORM_WINDOWS
 bool
 IsDebuggerPresent()
@@ -121,7 +126,8 @@ struct FLogFile
 #if UNSYNC_PLATFORM_WINDOWS
 		Handle = _wfopen(InFilename, L"wt, ccs=UNICODE");
 #else
-		// TODO
+		std::string FilenameUtf8 = ConvertWideToUtf8(std::wstring_view(InFilename));
+		Handle					 = fopen(FilenameUtf8.c_str(), "wt");
 #endif
 	}
 
@@ -139,6 +145,44 @@ struct FLogFile
 
 std::unique_ptr<FLogFile> GLogFile;
 
+static void
+RotateLogs(const FPath& LogFilename)
+{
+	static const FPath ExpectedSuffix(".log");
+
+	std::error_code ErrorCode;
+	FFileAttributes ExistingAttrib = GetFileAttrib(LogFilename);
+
+	if (!LogFilename.has_parent_path() || !ExistingAttrib.bValid)
+	{
+		return;
+	}
+
+	FTimestampString TimestampString	= FormatWindowsFileTime(ExistingAttrib.Mtime);
+	std::string		 BackupExtensionStr = TimestampString.data();
+	std::replace_if(
+		BackupExtensionStr.begin(),
+		BackupExtensionStr.end(),
+		[](char C) { return C == ':' || C == '.'; },
+		'-');
+	BackupExtensionStr += ".log";
+
+	FPath BackupExtension	= FPath(BackupExtensionStr);
+	FPath BackupLogFilename = LogFilename;
+	BackupLogFilename.replace_extension(BackupExtension);
+
+	FileRename(LogFilename, BackupLogFilename, ErrorCode);
+
+	FPath LogDirectory = LogFilename.parent_path();
+
+	FPathFilterCallback Filter = [](const FPath& CandidatePath)
+	{
+		return CandidatePath.has_extension() && CandidatePath.native().ends_with(ExpectedSuffix.native());
+	};
+
+	DeleteOldFilesInDirectory(LogDirectory, 10 /*files to keep*/, true /*allow in dry run*/, Filter);
+}
+
 void
 LogSetFileInternal(const wchar_t* Filename)
 {
@@ -156,6 +200,8 @@ FLogFileScope::FLogFileScope(const wchar_t* Filename)
 {
 	std::wstring CommandLine;
 
+	RotateLogs(FPath(Filename));
+
 	{
 		std::lock_guard<std::mutex> LockGuard(GLogMutex);
 		GLogFileStack.push_back(std::move(GLogFile));
@@ -242,10 +288,6 @@ LogFlush()
 	}
 }
 
-static constexpr size_t TimestampStringSize = sizeof("YYYY-MM-DDTHH:MM:SS.sssZ");
-
-using FTimestampString = fmt::basic_memory_buffer<char, TimestampStringSize>;
-
 static FTimestampString
 FormatTimestamp(std::chrono::system_clock::time_point Timestamp)
 {
@@ -263,6 +305,13 @@ FormatTimestamp(std::chrono::system_clock::time_point Timestamp)
 	return Result;
 }
 
+static FTimestampString
+FormatWindowsFileTime(uint64 Ticks)
+{
+	std::filesystem::file_time_type FileTime = FromWindowsFileTime(Ticks);
+	return FormatTimestamp(SystemTimeFromFileTime(FileTime));
+}
+
 void
 LogPrintf(ELogLevel Level, const wchar_t* Str, ...)
 {
diff --git a/Engine/Source/Programs/Unsync/Private/UnsyncVersion.cpp b/Engine/Source/Programs/Unsync/Private/UnsyncVersion.cpp
--- a/Engine/Source/Programs/Unsync/Private/UnsyncVersion.cpp
+++ b/Engine/Source/Programs/Unsync/Private/UnsyncVersion.cpp
@@ -2,7 +2,7 @@
 
 #include "UnsyncVersion.h"
 
-#define UNSYNC_VERSION_STR "1.0.86"
+#define UNSYNC_VERSION_STR "1.0.87"
 
 namespace unsync {
 
