unsync - Log Horde server diagnostic messages encountered during blob downloads

diff --git a/Engine/Source/Programs/Unsync/Private/UnsyncHorde.cpp b/Engine/Source/Programs/Unsync/Private/UnsyncHorde.cpp
--- a/Engine/Source/Programs/Unsync/Private/UnsyncHorde.cpp
+++ b/Engine/Source/Programs/Unsync/Private/UnsyncHorde.cpp
@@ -135,81 +135,111 @@ FDownloadResult FHordeProtocolImpl::DownloadArtifactBlobs(std::string_view Artif
 
 	while (Reader.RemainingSize())
 	{
-		FHordeUnsyncBlobHeaderV1 BlobHeader = {};
-		Reader.ReadInto(BlobHeader.Magic);
-		Reader.ReadInto(BlobHeader.PayloadSize);
-		Reader.ReadInto(BlobHeader.DecompressedSize);
-		Reader.ReadInto(BlobHeader.DecompressedHash);
+		uint64 HeaderId = 0;
+		Reader.ReadInto(HeaderId);
 
-		if (BlobHeader.Magic != FHordeUnsyncBlobHeaderV1::MAGIC)
+		if (HeaderId == FHordeUnsyncBlobHeaderV1::MAGIC)
 		{
-			UNSYNC_ERROR(L"Got unexpected blob header identifier");
-			return FDownloadError(EDownloadRetryMode::Abort);
-		}
+			FHordeUnsyncBlobHeaderV1 PacketHeader = {};
 
-		FDownloadedBlock DownloadedBlock;
-		DownloadedBlock.bCompressed = false;  // we always decompress the block before handing it over to the caller
+			PacketHeader.Magic = HeaderId;
+			Reader.ReadInto(PacketHeader.PayloadSize);
+			Reader.ReadInto(PacketHeader.DecompressedSize);
+			Reader.ReadInto(PacketHeader.DecompressedHash);
 
-		FBufferView Payload = Response.Buffer.View(Reader.Tell(), BlobHeader.PayloadSize);
-		FBuffer		DecompressedBuffer;
+			FDownloadedBlock DownloadedBlock;
+			DownloadedBlock.bCompressed = false;  // we always decompress the block before handing it over to the caller
 
-		if (ChunkContentEncoding == "zstd")
-		{
-			if (BlobHeader.DecompressedSize != 0)
+			FBufferView Payload = Response.Buffer.View(Reader.Tell(), PacketHeader.PayloadSize);
+			FBuffer		DecompressedBuffer;
+
+			if (ChunkContentEncoding == "zstd")
 			{
-				DecompressedBuffer				 = Decompress(Payload.Data, Payload.Size);
-				DownloadedBlock.Data			 = DecompressedBuffer.Data();
-				DownloadedBlock.DecompressedSize = DecompressedBuffer.Size();
+				if (PacketHeader.DecompressedSize != 0)
+				{
+					DecompressedBuffer				 = Decompress(Payload.Data, Payload.Size);
+					DownloadedBlock.Data			 = DecompressedBuffer.Data();
+					DownloadedBlock.DecompressedSize = DecompressedBuffer.Size();
+				}
+				else
+				{
+					DownloadedBlock.Data			 = DecompressedBuffer.Data();
+					DownloadedBlock.DecompressedSize = DecompressedBuffer.Size();
+				}
+			}
+			else if (ChunkContentEncoding == "" || ChunkContentEncoding == "identity")
+			{
+				DownloadedBlock.Data			 = Payload.Data;
+				DownloadedBlock.DecompressedSize = Payload.Size;
+
+				if (PacketHeader.DecompressedSize != PacketHeader.PayloadSize)
+				{
+					UNSYNC_ERROR(L"Received blob size (%llu bytes) does not match expected size (%llu bytes)",
+								 llu(PacketHeader.DecompressedSize),
+								 llu(PacketHeader.PayloadSize));
+					return FDownloadError(EDownloadRetryMode::Abort);
+				}
 			}
 			else
 			{
-				DownloadedBlock.Data			 = DecompressedBuffer.Data();
-				DownloadedBlock.DecompressedSize = DecompressedBuffer.Size();
+				std::string Value = std::string(ChunkContentEncoding);
+				UNSYNC_ERROR(L"Unexpected chunk content encoding: '%hs'", Value.c_str());
+				return FDownloadError(EDownloadRetryMode::Abort);
 			}
-		}
-		else if (ChunkContentEncoding == "" || ChunkContentEncoding == "identity")
-		{
-			DownloadedBlock.Data			 = Payload.Data;
-			DownloadedBlock.DecompressedSize = Payload.Size;
 
-			if (BlobHeader.DecompressedSize != BlobHeader.PayloadSize)
+			if (DownloadedBlock.DecompressedSize != PacketHeader.DecompressedSize)
 			{
 				UNSYNC_ERROR(L"Received blob size (%llu bytes) does not match expected size (%llu bytes)",
-							 llu(BlobHeader.DecompressedSize),
-							 llu(BlobHeader.PayloadSize));
+							 llu(DownloadedBlock.DecompressedSize),
+							 llu(PacketHeader.DecompressedSize));
 				return FDownloadError(EDownloadRetryMode::Abort);
 			}
+
+			FGenericHash BlockHash = ComputeHash(DownloadedBlock.Data, DownloadedBlock.DecompressedSize, StrongHasher);
+
+			if (BlockHash.ToHash160() != PacketHeader.DecompressedHash)
+			{
+				std::string BlockHashStr	= HashToHexString(BlockHash.ToHash160());
+				std::string ExpectedHashStr = HashToHexString(PacketHeader.DecompressedHash);
+				UNSYNC_ERROR(L"Received blob hash (%hs) does not match expected hash (%hs)", BlockHashStr.c_str(), ExpectedHashStr.c_str());
+				return FDownloadError(EDownloadRetryMode::Abort);
+			}
+
+			FHash128 BlockHash128 = BlockHash.ToHash128();
+
+			CompletionCallback(DownloadedBlock, BlockHash128);
+
+			Reader.Skip(PacketHeader.PayloadSize);
 		}
-		else
+		else if (HeaderId == FHordeUnsyncBlobErrorHeaderV1::MAGIC)
 		{
-			std::string Value = std::string(ChunkContentEncoding);
-			UNSYNC_ERROR(L"Unexpected chunk content encoding: '%hs'", Value.c_str());
-			return FDownloadError(EDownloadRetryMode::Abort);
-		}
+			FHordeUnsyncBlobErrorHeaderV1 PacketHeader = {};
 
-		if (DownloadedBlock.DecompressedSize != BlobHeader.DecompressedSize)
-		{
-			UNSYNC_ERROR(L"Received blob size (%llu bytes) does not match expected size (%llu bytes)",
-						 llu(DownloadedBlock.DecompressedSize),
-						 llu(BlobHeader.DecompressedSize));
-			return FDownloadError(EDownloadRetryMode::Abort);
-		}
+			PacketHeader.Magic = HeaderId;
+			Reader.ReadInto(PacketHeader.PayloadSize);
+
+			if (PacketHeader.PayloadSize > 1)
+			{
+				uint32 StringLength = PacketHeader.PayloadSize - 1;	 // Error string is zero-terminated
+
+				std::string ErrorString;
+				ErrorString.resize(StringLength);
+				Reader.Read(ErrorString.data(), StringLength);
 
-		FGenericHash BlockHash = ComputeHash(DownloadedBlock.Data, DownloadedBlock.DecompressedSize, StrongHasher);
+				UNSYNC_ERROR(L"Horde server error: %hs", ErrorString.c_str());
+			}
+			else
+			{
+				UNSYNC_ERROR(L"Unknown Horde server error");
+			}
 
-		if (BlockHash.ToHash160() != BlobHeader.DecompressedHash)
+			return FDownloadError(EDownloadRetryMode::Abort);
+		}
+		else
 		{
-			std::string BlockHashStr	= HashToHexString(BlockHash.ToHash160());
-			std::string ExpectedHashStr = HashToHexString(BlobHeader.DecompressedHash);
-			UNSYNC_ERROR(L"Received blob hash (%hs) does not match expected hash (%hs)", BlockHashStr.c_str(), ExpectedHashStr.c_str());
+			UNSYNC_ERROR(L"Got unexpected blob header identifier");
 			return FDownloadError(EDownloadRetryMode::Abort);
 		}
-
-		FHash128 BlockHash128 = BlockHash.ToHash128();
-
-		CompletionCallback(DownloadedBlock, BlockHash128);
-
-		Reader.Skip(BlobHeader.PayloadSize);
 	}
 
 	return ResultOk<FDownloadError>();
diff --git a/Engine/Source/Programs/Unsync/Private/UnsyncProtocol.h b/Engine/Source/Programs/Unsync/Private/UnsyncProtocol.h
--- a/Engine/Source/Programs/Unsync/Private/UnsyncProtocol.h
+++ b/Engine/Source/Programs/Unsync/Private/UnsyncProtocol.h
@@ -300,4 +300,12 @@ struct FHordeUnsyncBlobHeaderV1
 	FHash160 DecompressedHash = {};
 };
 
+struct FHordeUnsyncBlobErrorHeaderV1
+{
+	static constexpr uint64 MAGIC = 0x4C5C2AABA992DEADull;
+
+	uint64	 Magic			  = 0;
+	uint32	 PayloadSize	  = 0;
+};
+
 }  // namespace unsync
diff --git a/Engine/Source/Programs/Unsync/Private/UnsyncVersion.cpp b/Engine/Source/Programs/Unsync/Private/UnsyncVersion.cpp
--- a/Engine/Source/Programs/Unsync/Private/UnsyncVersion.cpp
+++ b/Engine/Source/Programs/Unsync/Private/UnsyncVersion.cpp
@@ -2,7 +2,7 @@
 
 #include "UnsyncVersion.h"
 
-#define UNSYNC_VERSION_STR "1.0.87"
+#define UNSYNC_VERSION_STR "1.0.88"
 
 namespace unsync {
 
