unsync - Report errors received from a proxy server and handle socket error codes

diff --git a/Engine/Source/Programs/Unsync/Private/UnsyncHttp.h b/Engine/Source/Programs/Unsync/Private/UnsyncHttp.h
--- a/Engine/Source/Programs/Unsync/Private/UnsyncHttp.h
+++ b/Engine/Source/Programs/Unsync/Private/UnsyncHttp.h
@@ -66,7 +66,7 @@ struct FHttpResponse
 
 	bool Success() const { return Code >= 200 && Code < 300; }
 
-	std::string_view AsStringView() const { return std::string_view((const char*)Buffer.Data(), Buffer.Size()); }
+	std::string_view AsStringView() const { return unsync::AsStringView(Buffer); }
 
 	std::string_view FindHeader(const std::string_view Name) const;
 };
diff --git a/Engine/Source/Programs/Unsync/Private/UnsyncProtocol.h b/Engine/Source/Programs/Unsync/Private/UnsyncProtocol.h
--- a/Engine/Source/Programs/Unsync/Private/UnsyncProtocol.h
+++ b/Engine/Source/Programs/Unsync/Private/UnsyncProtocol.h
@@ -308,4 +308,8 @@ struct FHordeUnsyncBlobErrorHeaderV1
 	uint32	 PayloadSize	  = 0;
 };
 
+// Block stream response is always terminated using a packet with this hash.
+// The packet payload may optionally contain a Json with diagnostics.
+inline const FHash128 TERMINATOR_BLOCK_HASH = FHash128{};
+
 }  // namespace unsync
diff --git a/Engine/Source/Programs/Unsync/Private/UnsyncProxy.cpp b/Engine/Source/Programs/Unsync/Private/UnsyncProxy.cpp
--- a/Engine/Source/Programs/Unsync/Private/UnsyncProxy.cpp
+++ b/Engine/Source/Programs/Unsync/Private/UnsyncProxy.cpp
@@ -294,6 +294,38 @@ SortBlockRequestsByFileName(std::vector<RequestType>& Requests)
 			  });
 }
 
+// Log any diagnostic messages. Return true if block download request can be retried.
+bool
+ProcessBlockPacketErrorJson(std::string_view DiagnosticViewJson)
+{
+	std::string DiagnosticStringJson(DiagnosticViewJson);
+
+	using namespace json11;
+
+	std::string JsonErrorString;
+	Json		JsonObject = Json::parse(DiagnosticStringJson, JsonErrorString);
+
+	if (!JsonErrorString.empty())
+	{
+		UNSYNC_ERROR(L"Failed to parse diagnostic message from UNSYNC server: %hs", JsonErrorString.c_str());
+		return false;
+	}
+
+	if (auto& Field = JsonObject["message"]; Field.is_string())
+	{
+		const std::string& FieldVal = Field.string_value();
+		UNSYNC_ERROR(L"Server error: %hs", FieldVal.c_str());
+	}
+
+	if (auto& Field = JsonObject["can_retry"]; Field.is_bool())
+	{
+		bool FieldVal = Field.bool_value();
+		return FieldVal;
+	}
+
+	return false;
+}
+
 FDownloadResult
 FUnsyncProtocolImpl::Download(const TArrayView<FNeedBlock> NeedBlocks, const FBlockDownloadCallback& CompletionCallback)
 {
@@ -431,11 +463,26 @@ FUnsyncProtocolImpl::Download(const TArrayView<FNeedBlock> NeedBlocks, const FBl
 			break;
 		}
 
-		static const FHash128 TerminatorHash = FHash128{};	// response is always terminated with an empty packet
-
-		if (BlockPacket.Hash == TerminatorHash)
+		if (BlockPacket.Hash == TERMINATOR_BLOCK_HASH)
 		{
-			// TODO: Termination packet payload may contain diagnostic messages that can be reported to the user
+			// If we received the terminator packet before the last expected block, it likely indicates an error of some kind.
+			if (I < Requests.size())
+			{
+				std::string_view DiagnosticView = AsStringView(BlockPacket.Data);
+				if (DiagnosticView.starts_with('{'))
+				{
+					bool bCanRetry = ProcessBlockPacketErrorJson(DiagnosticView);
+					if (!bCanRetry)
+					{
+						return FDownloadError(EDownloadRetryMode::Abort);
+					}
+				}
+				else
+				{
+					UNSYNC_WARNING(L"Received early stream termination packet without diagnostics");
+				}
+			}
+
 			break;
 		}
 
diff --git a/Engine/Source/Programs/Unsync/Private/UnsyncSocket.cpp b/Engine/Source/Programs/Unsync/Private/UnsyncSocket.cpp
--- a/Engine/Source/Programs/Unsync/Private/UnsyncSocket.cpp
+++ b/Engine/Source/Programs/Unsync/Private/UnsyncSocket.cpp
@@ -98,6 +98,23 @@ GetLastSocketError()
 #endif
 }
 
+static void
+ReportSocketError(int32 ErrorCode)
+{
+#if UNSYNC_PLATFORM_WINDOWS
+	ErrorCode = WSAGetLastError();
+#endif
+
+	UNSYNC_ERROR(L"Socket error: %d", ErrorCode);
+}
+
+static void
+ReportSocketError(tls* TlsCtx)
+{
+	const char* ErrorString = tls_error(TlsCtx);
+	UNSYNC_ERROR(L"TLS error: %hs", ErrorString);
+}
+
 FSocketHandle
 SocketListenTcp(const char* Address, uint16 Port)
 {
@@ -233,12 +250,23 @@ int
 SocketSend(FSocketHandle Socket, const void* Data, size_t DataSize)
 {
 	if (DataSize == 0)
+	{
 		return 0;
+	}
 
 	UNSYNC_ASSERT(DataSize < INT_MAX);
 
 	int Result = send(Socket, reinterpret_cast<const char*>(Data), (int)DataSize, 0);
-	return Result;
+
+	if (Result < 0)
+	{
+		ReportSocketError(Result);
+		return 0;
+	}
+	else
+	{
+		return Result;
+	}
 }
 
 int
@@ -260,6 +288,10 @@ SocketRecvAll(FSocketHandle Socket, void* Data, size_t DataSize)
 		int	  Res = recv(Socket, BatchPtr, BatchSize, MSG_WAITALL);
 		if (Res <= 0)
 		{
+			if (Res < 0)
+			{
+				ReportSocketError(Res);
+			}
 			break;
 		}
 		ProcessedBytes += Res;
@@ -276,7 +308,15 @@ SocketRecvAny(FSocketHandle Socket, void* Data, size_t DataSize)
 
 	UNSYNC_ASSERT(DataSize < INT_MAX);
 
-	return recv(Socket, reinterpret_cast<char*>(Data), (int)DataSize, 0);
+	int Res = recv(Socket, reinterpret_cast<char*>(Data), (int)DataSize, 0);
+
+	if (Res < 0)
+	{
+		ReportSocketError(Res);
+		return 0;
+	}
+
+	return Res;
 }
 
 bool
@@ -432,11 +472,16 @@ FSocketTls::Send(const void* Data, size_t DataSize)
 	int ProcessedBytes = 0;
 	while (ProcessedBytes < DataSize)
 	{
-		int Res = CheckedNarrow(tls_write(TlsCtx, (const char*)Data + ProcessedBytes, DataSize - ProcessedBytes));
-		if (Res <= 0)
+		ssize_t ErrorCode = tls_write(TlsCtx, (const char*)Data + ProcessedBytes, DataSize - ProcessedBytes);
+		if (ErrorCode <= 0)
 		{
+			if (ErrorCode < 0)
+			{
+				ReportSocketError(TlsCtx);
+			}
 			break;
 		}
+		int Res = CheckedNarrow(ErrorCode);
 		ProcessedBytes += Res;
 	}
 
@@ -454,11 +499,16 @@ FSocketTls::RecvAll(void* Data, size_t DataSize)
 	int ProcessedBytes = 0;
 	while (ProcessedBytes < DataSize)
 	{
-		int Res = CheckedNarrow(tls_read(TlsCtx, (char*)Data + ProcessedBytes, DataSize - ProcessedBytes));
-		if (Res <= 0)
+		ssize_t ErrorCode = tls_read(TlsCtx, (char*)Data + ProcessedBytes, DataSize - ProcessedBytes);
+		if (ErrorCode <= 0)
 		{
+			if (ErrorCode < 0)
+			{
+				ReportSocketError(TlsCtx);
+			}
 			break;
 		}
+		int Res = CheckedNarrow(ErrorCode);
 		ProcessedBytes += Res;
 	}
 
@@ -476,11 +526,16 @@ FSocketTls::RecvAny(void* Data, size_t DataSize)
 	int ProcessedBytes = 0;
 	do
 	{
-		int Res = CheckedNarrow(tls_read(TlsCtx, (char*)Data + ProcessedBytes, DataSize - ProcessedBytes));
-		if (Res <= 0)
+		ssize_t ErrorCode = tls_read(TlsCtx, (char*)Data + ProcessedBytes, DataSize - ProcessedBytes);
+		if (ErrorCode <= 0)
 		{
+			if (ErrorCode < 0)
+			{
+				ReportSocketError(TlsCtx);
+			}
 			break;
 		}
+		int Res = CheckedNarrow(ErrorCode);
 		ProcessedBytes += Res;
 	} while (false);
 
diff --git a/Engine/Source/Programs/Unsync/Private/UnsyncUtil.cpp b/Engine/Source/Programs/Unsync/Private/UnsyncUtil.cpp
--- a/Engine/Source/Programs/Unsync/Private/UnsyncUtil.cpp
+++ b/Engine/Source/Programs/Unsync/Private/UnsyncUtil.cpp
@@ -815,4 +815,17 @@ SplitByAny(std::string_view String, const char* SeparatorCharacters)
 	return Result;
 }
 
+std::string_view
+AsStringView(const FBuffer& Buffer)
+{
+	if (Buffer.Empty())
+	{
+		return {};
+	}
+	else
+	{
+		return std::string_view(reinterpret_cast<const char*>(Buffer.Data()), Buffer.Size());
+	}
+}
+
 }  // namespace unsync
diff --git a/Engine/Source/Programs/Unsync/Private/UnsyncUtil.h b/Engine/Source/Programs/Unsync/Private/UnsyncUtil.h
--- a/Engine/Source/Programs/Unsync/Private/UnsyncUtil.h
+++ b/Engine/Source/Programs/Unsync/Private/UnsyncUtil.h
@@ -140,6 +140,8 @@ std::wstring ConvertUtf8ToWide(std::string_view StringUtf8);
 std::string	 ConvertWideToUtf8(std::wstring_view StringWide);
 void		 ConvertWideToUtf8(std::wstring_view StringWide, std::string& Result);
 
+std::string_view AsStringView(const FBuffer& Buffer);
+
 inline std::string
 ToString(const std::wstring_view WideStringView)
 {
diff --git a/Engine/Source/Programs/Unsync/Private/UnsyncVersion.cpp b/Engine/Source/Programs/Unsync/Private/UnsyncVersion.cpp
--- a/Engine/Source/Programs/Unsync/Private/UnsyncVersion.cpp
+++ b/Engine/Source/Programs/Unsync/Private/UnsyncVersion.cpp
@@ -2,7 +2,7 @@
 
 #include "UnsyncVersion.h"
 
-#define UNSYNC_VERSION_STR "1.0.88"
+#define UNSYNC_VERSION_STR "1.0.89"
 
 namespace unsync {
 
