unsync - Allow retrying requests by default

diff --git a/Engine/Source/Programs/Unsync/Private/UnsyncProxy.cpp b/Engine/Source/Programs/Unsync/Private/UnsyncProxy.cpp
--- a/Engine/Source/Programs/Unsync/Private/UnsyncProxy.cpp
+++ b/Engine/Source/Programs/Unsync/Private/UnsyncProxy.cpp
@@ -323,7 +323,8 @@ ProcessBlockPacketErrorJson(std::string_view DiagnosticViewJson)
 		return FieldVal;
 	}
 
-	return false;
+	// Allow retrying the request by default
+	return true;
 }
 
 FDownloadResult
