unsync - Downgrade socket diagnostics from errors to warnings and use callbacks for TLS socket read/write

diff --git a/Engine/Source/Programs/Unsync/Private/UnsyncLog.h b/Engine/Source/Programs/Unsync/Private/UnsyncLog.h
--- a/Engine/Source/Programs/Unsync/Private/UnsyncLog.h
+++ b/Engine/Source/Programs/Unsync/Private/UnsyncLog.h
@@ -158,4 +158,17 @@ void LogSaveCommandLineUtf8(int Argc, char** Argv);
 			UNSYNC_BREAK();     \
 	}
 
+#define UNSYNC_BREAK_ON_ERROR_LEVEL(LogLevel)              \
+	{                                                      \
+		if (LogLevel == ELogLevel::Error && GBreakOnError) \
+			UNSYNC_BREAK();                                \
+	}
+
+#define UNSYNC_BREAK_ON_WARNING_LEVEL(LogLevel)                \
+	{                                                          \
+		if (LogLevel == ELogLevel::Warning && GBreakOnWarning) \
+			UNSYNC_BREAK();                                    \
+	}
+
+
 }  // namespace unsync
diff --git a/Engine/Source/Programs/Unsync/Private/UnsyncSocket.cpp b/Engine/Source/Programs/Unsync/Private/UnsyncSocket.cpp
--- a/Engine/Source/Programs/Unsync/Private/UnsyncSocket.cpp
+++ b/Engine/Source/Programs/Unsync/Private/UnsyncSocket.cpp
@@ -99,20 +99,22 @@ GetLastSocketError()
 }
 
 static void
-ReportSocketError(int32 ErrorCode)
+ReportSocketError(int32 ErrorCode, ELogLevel LogLevel)
 {
 #if UNSYNC_PLATFORM_WINDOWS
 	ErrorCode = WSAGetLastError();
 #endif
 
-	UNSYNC_ERROR(L"Socket error: %d", ErrorCode);
+	LogPrintf(LogLevel, L"Socket error: %d", ErrorCode);
+	UNSYNC_BREAK_ON_ERROR_LEVEL(LogLevel);
 }
 
 static void
-ReportSocketError(tls* TlsCtx)
+ReportSocketError(tls* TlsCtx, ELogLevel LogLevel)
 {
 	const char* ErrorString = tls_error(TlsCtx);
-	UNSYNC_ERROR(L"TLS error: %hs", ErrorString);
+	LogPrintf(LogLevel, L"TLS error: %hs", ErrorString);
+	UNSYNC_BREAK_ON_ERROR_LEVEL(LogLevel);
 }
 
 FSocketHandle
@@ -260,7 +262,7 @@ SocketSend(FSocketHandle Socket, const void* Data, size_t DataSize)
 
 	if (Result < 0)
 	{
-		ReportSocketError(Result);
+		ReportSocketError(Result, ELogLevel::Warning);
 		return 0;
 	}
 	else
@@ -290,7 +292,7 @@ SocketRecvAll(FSocketHandle Socket, void* Data, size_t DataSize)
 		{
 			if (Res < 0)
 			{
-				ReportSocketError(Res);
+				ReportSocketError(Res, ELogLevel::Warning);
 			}
 			break;
 		}
@@ -312,7 +314,7 @@ SocketRecvAny(FSocketHandle Socket, void* Data, size_t DataSize)
 
 	if (Res < 0)
 	{
-		ReportSocketError(Res);
+		ReportSocketError(Res, ELogLevel::Warning);
 		return 0;
 	}
 
@@ -372,6 +374,20 @@ FSocketRaw::RecvAny(void* Data, size_t DataSize)
 	return SocketRecvAny(Handle, Data, DataSize);
 }
 
+static ssize_t
+TlsReadCallback(struct tls* TlsCtx, void* Buffer, size_t BufferSize, void* UserData)
+{
+	FSocketHandle Socket = reinterpret_cast<FSocketTls*>(UserData)->Handle;
+	return SocketRecvAll(Socket, Buffer, BufferSize);
+}
+
+static ssize_t
+TlsWriteCallback(struct tls* TlsCtx, const void* Buffer, size_t BufferSize, void* UserData)
+{
+	FSocketHandle Socket = reinterpret_cast<FSocketTls*>(UserData)->Handle;
+	return SocketSend(Socket, Buffer, BufferSize);
+}
+
 FSocketTls::FSocketTls(FSocketHandle InHandle, FTlsClientSettings ClientSettings) : FSocketBase(InHandle)
 {
 	Security = ESocketSecurity::Unknown;
@@ -413,7 +429,7 @@ FSocketTls::FSocketTls(FSocketHandle InHandle, FTlsClientSettings ClientSettings
 	{
 		UNSYNC_ASSERT(!ClientSettings.Subject.empty());
 		std::string TlsSubject(ClientSettings.Subject);
-		Err = tls_connect_socket(TlsCtx, (int)Handle, TlsSubject.c_str());
+		Err = tls_connect_cbs(TlsCtx, TlsReadCallback, TlsWriteCallback, this, TlsSubject.c_str());
 	}
 
 	if (Err == 0)
@@ -477,7 +493,7 @@ FSocketTls::Send(const void* Data, size_t DataSize)
 		{
 			if (ErrorCode < 0)
 			{
-				ReportSocketError(TlsCtx);
+				ReportSocketError(TlsCtx, ELogLevel::Warning);
 			}
 			break;
 		}
@@ -504,7 +520,7 @@ FSocketTls::RecvAll(void* Data, size_t DataSize)
 		{
 			if (ErrorCode < 0)
 			{
-				ReportSocketError(TlsCtx);
+				ReportSocketError(TlsCtx, ELogLevel::Warning);
 			}
 			break;
 		}
@@ -531,7 +547,7 @@ FSocketTls::RecvAny(void* Data, size_t DataSize)
 		{
 			if (ErrorCode < 0)
 			{
-				ReportSocketError(TlsCtx);
+				ReportSocketError(TlsCtx, ELogLevel::Warning);
 			}
 			break;
 		}
diff --git a/Engine/Source/Programs/Unsync/Private/UnsyncVersion.cpp b/Engine/Source/Programs/Unsync/Private/UnsyncVersion.cpp
--- a/Engine/Source/Programs/Unsync/Private/UnsyncVersion.cpp
+++ b/Engine/Source/Programs/Unsync/Private/UnsyncVersion.cpp
@@ -2,7 +2,7 @@
 
 #include "UnsyncVersion.h"
 
-#define UNSYNC_VERSION_STR "1.0.89"
+#define UNSYNC_VERSION_STR "1.0.90"
 
 namespace unsync {
 
