unsync - Obfuscate cached tokens on Windows

* Use CryptProtectData() to encrypt the token using user-specific key

diff --git a/Engine/Source/Programs/Unsync/Private/UnsyncAuth.cpp b/Engine/Source/Programs/Unsync/Private/UnsyncAuth.cpp
--- a/Engine/Source/Programs/Unsync/Private/UnsyncAuth.cpp
+++ b/Engine/Source/Programs/Unsync/Private/UnsyncAuth.cpp
@@ -16,6 +16,15 @@
 #include <openssl/rand.h>
 #include <openssl/sha.h>
 
+#if UNSYNC_PLATFORM_WINDOWS
+#	define UNSYNC_OBFUSCATE_CACHED_AUTH_TOKEN 1
+#	include <Windows.h>
+#	include <Dpapi.h>
+#	pragma comment(lib, "Crypt32.lib")
+#else
+#	define UNSYNC_OBFUSCATE_CACHED_AUTH_TOKEN 0
+#endif
+
 namespace unsync {
 
 std::string
@@ -646,19 +655,130 @@ struct FAuthTokenCache
 
 static FAuthTokenCache GAuthTokenCache;
 
+enum class EProtectedBufferFormat : uint64 {
+	Invalid				  = 0,
+	Win32CryptProtectData = 0x9E9AA2B319A7D98Full,
+};
+
+struct FProtectedBufferHeader
+{
+	EProtectedBufferFormat Format = EProtectedBufferFormat::Invalid;
+	uint64				   Size	  = 0;
+};
+
+#if UNSYNC_OBFUSCATE_CACHED_AUTH_TOKEN
+#	if UNSYNC_PLATFORM_WINDOWS
+static TResult<FBuffer>
+ProtectBuffer(const FBufferView& InPlainTextBuffer)
+{
+	DATA_BLOB Blob = {};
+	Blob.cbData	   = (DWORD)InPlainTextBuffer.Size;
+	Blob.pbData	   = (BYTE*)InPlainTextBuffer.Data;
+
+	DATA_BLOB ProtectedBlob = {};
+
+	BOOL ProtectOk = CryptProtectData(&Blob, nullptr, nullptr, nullptr, nullptr, 0, &ProtectedBlob);
+	if (!ProtectOk)
+	{
+		DWORD ErrorCode = GetLastError();
+		return SystemError(L"CryptProtectData failed", ErrorCode);
+	}
+
+	FBuffer Result;
+
+	FProtectedBufferHeader Header;
+	Header.Format = EProtectedBufferFormat::Win32CryptProtectData;
+	Header.Size	  = ProtectedBlob.cbData;
+
+	Result.Append(reinterpret_cast<const uint8*>(&Header), sizeof(Header));
+	Result.Append(reinterpret_cast<const uint8*>(ProtectedBlob.pbData), ProtectedBlob.cbData);
+
+	LocalFree(ProtectedBlob.pbData);
+
+	return ResultOk(std::move(Result));
+}
+
+static TResult<FBuffer>
+UnprotectBuffer(const FBufferView& InProtectedBuffer)
+{
+	FProtectedBufferHeader Header;
+
+	if (InProtectedBuffer.Size < sizeof(Header))
+	{
+		return AppError(L"Protected buffer does not contain a valid header");
+	}
+
+	memcpy(&Header, InProtectedBuffer.Data, sizeof(Header));
+
+	if (Header.Format != EProtectedBufferFormat::Win32CryptProtectData)
+	{
+		return AppError(L"Protected buffer format is not supported");
+	}
+
+	DATA_BLOB ProtectedBlob = {};
+	ProtectedBlob.cbData	= (DWORD)Header.Size;
+	ProtectedBlob.pbData	= (BYTE*)(InProtectedBuffer.Data + sizeof(Header));
+
+	DATA_BLOB Blob = {};
+
+	BOOL ProtectOk = CryptUnprotectData(&ProtectedBlob, nullptr, nullptr, nullptr, nullptr, 0, &Blob);
+	if (!ProtectOk)
+	{
+		DWORD ErrorCode = GetLastError();
+		return SystemError(L"CryptUnprotectData failed", ErrorCode);
+	}
+
+	FBuffer Result;
+	Result.Append(reinterpret_cast<const uint8*>(Blob.pbData), Blob.cbData);
+
+	LocalFree(Blob.pbData);
+
+	return ResultOk(std::move(Result));
+}
+#	endif	// UNSYNC_PLATFORM_WINDOWS
+#endif		// UNSYNC_OBFUSCATE_CACHED_AUTH_TOKEN
+
 bool
 SaveAuthToken(const FPath& Path, const FAuthToken& AuthToken)
 {
-	bool bWrittenOk = WriteBufferToFile(Path,
-										(const uint8*)AuthToken.Raw.data(),
-										AuthToken.Raw.length(),
-										EFileMode::CreateWriteOnly | EFileMode::IgnoreDryRun);
+#if UNSYNC_OBFUSCATE_CACHED_AUTH_TOKEN
+	FBufferView		 AuthTokenView	 = {.Data = (const uint8*)AuthToken.Raw.data(), .Size = AuthToken.Raw.length()};
+	TResult<FBuffer> ProtectedBuffer = ProtectBuffer(AuthTokenView);
+	if (ProtectedBuffer.IsError())
+	{
+		LogError(ProtectedBuffer.GetError(), L"Failed to protect authentication token");
+		return false;
+	}
+	const bool bWrittenOk = WriteBufferToFile(Path,
+											  ProtectedBuffer.GetData(),
+											  EFileMode::CreateWriteOnly | EFileMode::IgnoreDryRun);
+#else	// UNSYNC_OBFUSCATE_CACHED_AUTH_TOKEN
+	const bool bWrittenOk = WriteBufferToFile(Path,
+											  (const uint8*)AuthToken.Raw.data(),
+											  AuthToken.Raw.length(),
+											  EFileMode::CreateWriteOnly | EFileMode::IgnoreDryRun);
+#endif	// UNSYNC_OBFUSCATE_CACHED_AUTH_TOKEN
 
 	if (!bWrittenOk)
 	{
 		return false;
 	}
 
+	{
+		FPath ExtendedPath = MakeExtendedAbsolutePath(Path);
+		std::error_code ErrorCode = {};
+		std::filesystem::permissions(
+				ExtendedPath,
+				std::filesystem::perms::owner_write | std::filesystem::perms::owner_read,
+				std::filesystem::perm_options::replace,
+				ErrorCode);
+		if (ErrorCode)
+		{
+			UNSYNC_ERROR(L"Failed to set cached token file permissions");
+			return false;
+		}
+	}
+
 	FFileAttributes Attrib = GetFileAttrib(Path);
 	if (!Attrib.bValid)
 	{
@@ -670,6 +790,24 @@ SaveAuthToken(const FPath& Path, const FAuthToken& AuthToken)
 	return true;
 }
 
+static bool
+IsProtectedBuffer(const FBufferView& BufferView)
+{
+	EProtectedBufferFormat Magic = EProtectedBufferFormat::Invalid;
+	if (BufferView.Size < sizeof(uint64))
+	{
+		return false;
+	}
+
+	memcpy(&Magic, BufferView.Data, sizeof(Magic));
+	if (Magic == EProtectedBufferFormat::Win32CryptProtectData)
+	{
+		return true;
+	}
+
+	return false;
+}
+
 TResult<FAuthToken>
 LoadAuthToken(const FPath& Path)
 {
@@ -679,6 +817,24 @@ LoadAuthToken(const FPath& Path)
 	}
 
 	FBuffer FileBuffer = ReadFileToBuffer(Path);
+
+	if (IsProtectedBuffer(FileBuffer))
+	{
+#if UNSYNC_OBFUSCATE_CACHED_AUTH_TOKEN
+		TResult<FBuffer> UnprotectResult = UnprotectBuffer(FileBuffer);
+		if (UnprotectResult.IsOk())
+		{
+			FileBuffer = std::move(UnprotectResult.GetData());
+		}
+		else
+		{
+			return MoveError<FAuthToken>(UnprotectResult);
+		}
+#else // UNSYNC_OBFUSCATE_CACHED_AUTH_TOKEN
+		return AppError(L"Protected token format is not supported");
+#endif	// UNSYNC_OBFUSCATE_CACHED_AUTH_TOKEN
+	}
+
 	if (FileBuffer.Size())
 	{
 		using namespace json11;
diff --git a/Engine/Source/Programs/Unsync/Private/UnsyncVersion.cpp b/Engine/Source/Programs/Unsync/Private/UnsyncVersion.cpp
--- a/Engine/Source/Programs/Unsync/Private/UnsyncVersion.cpp
+++ b/Engine/Source/Programs/Unsync/Private/UnsyncVersion.cpp
@@ -2,7 +2,7 @@
 
 #include "UnsyncVersion.h"
 
-#define UNSYNC_VERSION_STR "1.0.90"
+#define UNSYNC_VERSION_STR "1.0.91"
 
 namespace unsync {
 
