Unsync: Add support for setting something to executable.

On posix platforms we need to explicitly set the +x bit.

diff --git a/Engine/Source/Programs/Unsync/Private/UnsyncCmdInfo.cpp b/Engine/Source/Programs/Unsync/Private/UnsyncCmdInfo.cpp
--- a/Engine/Source/Programs/Unsync/Private/UnsyncCmdInfo.cpp
+++ b/Engine/Source/Programs/Unsync/Private/UnsyncCmdInfo.cpp
@@ -148,6 +148,7 @@ static void LogDecodedManifestJson(const FDirectoryManifest& Manifest)
 		Output += L"{";
 		FormatJsonKeyValueStr(Output, L"name", Name, L",");
 		FormatJsonKeyValueBool(Output, L"read_only", FileManifest.bReadOnly, L", ");
+		FormatJsonKeyValueBool(Output, L"executable", FileManifest.bIsExecutable, L", ");
 		FormatJsonKeyValueUInt(Output, L"mtime", FileManifest.Mtime, L", ");
 		FormatJsonKeyValueUInt(Output, L"size", FileManifest.Size, L", ");
 		FormatJsonKeyValueUInt(Output, L"block_size", FileManifest.BlockSize, L", ");
diff --git a/Engine/Source/Programs/Unsync/Private/UnsyncCmdPack.cpp b/Engine/Source/Programs/Unsync/Private/UnsyncCmdPack.cpp
--- a/Engine/Source/Programs/Unsync/Private/UnsyncCmdPack.cpp
+++ b/Engine/Source/Programs/Unsync/Private/UnsyncCmdPack.cpp
@@ -380,9 +380,10 @@ CmdPack(const FCmdPackOptions& Options)
 				{
 					if (It.second.RevisionControlIdentity == OldFileIt->second.RevisionControlIdentity)
 					{
-						It.second.Mtime		= OldFileIt->second.Mtime;
-						It.second.Size		= OldFileIt->second.Size;
-						It.second.bReadOnly = OldFileIt->second.bReadOnly;
+						It.second.Mtime		    = OldFileIt->second.Mtime;
+						It.second.Size		    = OldFileIt->second.Size;
+						It.second.bReadOnly     = OldFileIt->second.bReadOnly;
+						It.second.bIsExecutable = OldFileIt->second.bIsExecutable;
 					}
 				}
 			}
@@ -391,9 +392,10 @@ CmdPack(const FCmdPackOptions& Options)
 			{
 				FFileAttributes Attrib = GetFileAttrib(It.second.CurrentPath);
 
-				It.second.Mtime		= Attrib.Mtime;
-				It.second.Size		= Attrib.Size;
-				It.second.bReadOnly = true;	 // treat all p4 files as read-only in the manifest
+				It.second.Mtime		    = Attrib.Mtime;
+				It.second.Size		    = Attrib.Size;
+				It.second.bReadOnly     = true;	 // treat all p4 files as read-only in the manifest
+				It.second.bIsExecutable = Attrib.bIsExecutable;
 			}
 		};
 		ParallelForEach(DirectoryManifest.Files, UpdateFileMetadata);
@@ -791,7 +793,7 @@ SyncDirectoryFromPack(const FPath& OutputRoot, const FPackDatabase& PackDb, cons
 
 				// Check if the file on disk still matches the manifest and invalidate it otherwise
 				if (!Attrib.bValid || It.second.Mtime != Attrib.Mtime || It.second.Size != Attrib.Size ||
-					It.second.bReadOnly != Attrib.bReadOnly)
+					It.second.bReadOnly != Attrib.bReadOnly || It.second.bIsExecutable != Attrib.bIsExecutable)
 				{
 					It.second = FFileManifest();
 					NumDirtyFiles++;
@@ -876,7 +878,8 @@ SyncDirectoryFromPack(const FPath& OutputRoot, const FPackDatabase& PackDb, cons
 			// TODO: fix only the the read-only flag when that's the only difference
 			if (OldFileManifest->RevisionControlIdentity == FileManifest.RevisionControlIdentity &&
 				OldFileManifest->Mtime == FileManifest.Mtime && OldFileManifest->Size == FileManifest.Size &&
-				OldFileManifest->bReadOnly == FileManifest.bReadOnly)
+				OldFileManifest->bReadOnly == FileManifest.bReadOnly &&
+				OldFileManifest->bIsExecutable == FileManifest.bIsExecutable)
 			{
 				UNSYNC_VERBOSE2(L"Skipped '%ls' (up to date)", FileIt.first.c_str());
 				Stats.NumFilesSkipped += 1;
@@ -940,6 +943,11 @@ SyncDirectoryFromPack(const FPath& OutputRoot, const FPackDatabase& PackDb, cons
 			{
 				SetFileReadOnly(TargetFilePath, true);
 			}
+
+			if (FileManifest.bIsExecutable)
+			{
+				SetFileExecutable(TargetFilePath, true);
+			}
 		}
 
 		Stats.NumFilesSynced += 1;
diff --git a/Engine/Source/Programs/Unsync/Private/UnsyncCore.cpp b/Engine/Source/Programs/Unsync/Private/UnsyncCore.cpp
--- a/Engine/Source/Programs/Unsync/Private/UnsyncCore.cpp
+++ b/Engine/Source/Programs/Unsync/Private/UnsyncCore.cpp
@@ -2001,6 +2001,11 @@ SyncDirectory(const FSyncDirectoryOptions& SyncOptions)
 					{
 						SetFileReadOnly(Item.TargetFilePath, true);
 					}
+
+					if (Item.SourceManifest->bIsExecutable)
+					{
+						SetFileExecutable(Item.TargetFilePath, true);
+					}
 				}
 
 				if (bBackground)
diff --git a/Engine/Source/Programs/Unsync/Private/UnsyncFile.cpp b/Engine/Source/Programs/Unsync/Private/UnsyncFile.cpp
--- a/Engine/Source/Programs/Unsync/Private/UnsyncFile.cpp
+++ b/Engine/Source/Programs/Unsync/Private/UnsyncFile.cpp
@@ -964,11 +964,14 @@ GetFileAttrib(const FPath& Path, FFileAttributeCache* AttribCache)
 			return Result;
 		}
 
-		Result.bDirectory = Entry.is_directory();
-		Result.Size		  = Result.bDirectory ? 0 : Entry.file_size();
-		Result.Mtime	  = ToWindowsFileTime(Entry.last_write_time());
-		Result.bReadOnly  = IsReadOnly(Status.permissions());
-		Result.bValid	  = true;
+		std::filesystem::perms Perms = Status.permissions();
+
+		Result.bDirectory    = Entry.is_directory();
+		Result.Size		     = Result.bDirectory ? 0 : Entry.file_size();
+		Result.Mtime	     = ToWindowsFileTime(Entry.last_write_time());
+		Result.bReadOnly     = IsReadOnly(Perms);
+		Result.bIsExecutable = IsExecutable(Perms);
+		Result.bValid	     = true;
 	}
 
 	return Result;
@@ -1038,6 +1041,34 @@ SetFileReadOnly(const FPath& Path, bool bReadOnly)
 	return !ErrorCode;
 }
 
+bool
+SetFileExecutable(const FPath& Path, bool Executable)
+{
+	UNSYNC_ASSERT(!GDryRun);
+
+	FPath ExtendedPath = MakeExtendedAbsolutePath(Path);
+
+	std::error_code ErrorCode;
+	// TODO: In general if the +x bit is set, it set for all, this isn't always true though
+	// so we should at some point respect what the original file had set. For now though,
+	// this is sufficient.
+	if (Executable)
+	{
+		std::filesystem::permissions(
+			ExtendedPath,
+			std::filesystem::perms::owner_exec | std::filesystem::perms::group_exec | std::filesystem::perms::others_exec,
+			std::filesystem::perm_options::add, ErrorCode);
+	}
+	else
+	{
+		std::filesystem::permissions(ExtendedPath,
+			std::filesystem::perms::owner_exec | std::filesystem::perms::group_exec | std::filesystem::perms::others_exec,
+			std::filesystem::perm_options::remove, ErrorCode);
+	}
+
+	return !ErrorCode;
+}
+
 FBuffer
 ReadFileToBuffer(const FPath& Filename)
 {
@@ -1285,11 +1316,13 @@ CreateFileAttributeCache(const FPath& Root, const FSyncFilter* SyncFilter)
 		}
 
 		FFileAttributes Attr = {};
+		std::filesystem::perms Perms = Dir.status().permissions();
 
 		Attr.Mtime	= ToWindowsFileTime(Dir.last_write_time());
 		Attr.Size	= Dir.file_size();
 		Attr.bValid = true;
-		Attr.bReadOnly = IsReadOnly(Dir.status().permissions());
+		Attr.bReadOnly = IsReadOnly(Perms);
+		Attr.bIsExecutable = IsExecutable(Perms);
 
 		Result.Map[Dir.path().native()] = Attr;
 
@@ -1619,6 +1652,23 @@ TestFileAttrib()
 	const FFileAttributes FileAttribNonReadOnly = GetFileAttrib(TestFilename);
 	UNSYNC_ASSERT(!FileAttribNonReadOnly.bReadOnly);
 
+	const bool bIsExecutable = SetFileExecutable(TestFilename, true);
+	UNSYNC_ASSERT(bIsExecutable);
+
+	const FFileAttributes FileAttribExecutable = GetFileAttrib(TestFilename);
+	UNSYNC_ASSERT(FileAttribExecutable.bIsExecutable);
+
+#if UNSYNC_PLATFORM_UNIX
+	const bool bIsNotExecutable = SetFileExecutable(TestFilename, false);
+	UNSYNC_ASSERT(bIsNotExecutable);
+
+	// This part of the test will fail on Windows platforms as the +x bit
+	// means nothing there, so bIsExecutable will never be set to true.
+	const FFileAttributes FileAttribNotExecutable = GetFileAttrib(TestFilename);
+	UNSYNC_ASSERT(!FileAttribNotExecutable.bIsExecutable);
+
+#endif // UNSYNC_PLATFORM_UNIX
+
 	std::error_code ErrorCode;
 	const bool		bFileDeleted = FileRemove(TestFilename, ErrorCode);
 	UNSYNC_ASSERT(bFileDeleted);
diff --git a/Engine/Source/Programs/Unsync/Private/UnsyncFile.h b/Engine/Source/Programs/Unsync/Private/UnsyncFile.h
--- a/Engine/Source/Programs/Unsync/Private/UnsyncFile.h
+++ b/Engine/Source/Programs/Unsync/Private/UnsyncFile.h
@@ -460,11 +460,12 @@ bool	WriteBufferToFile(const FPath& Filename, const std::string& Buffer, EFileMo
 
 struct FFileAttributes
 {
-	uint64 Mtime	  = 0;	// Windows file time (100ns ticks since 1601-01-01T00:00:00Z)
-	uint64 Size		  = 0;
-	bool   bDirectory = false;
-	bool   bValid	  = false;
-	bool   bReadOnly  = false;
+	uint64 Mtime	     = 0;	// Windows file time (100ns ticks since 1601-01-01T00:00:00Z)
+	uint64 Size		     = 0;
+	bool   bDirectory    = false;
+	bool   bValid	     = false;
+	bool   bReadOnly     = false;
+	bool   bIsExecutable = false;
 };
 
 struct FFileAttributeCache
@@ -477,7 +478,17 @@ struct FFileAttributeCache
 inline bool
 IsReadOnly(std::filesystem::perms Perms)
 {
-	return (Perms & std::filesystem::perms::owner_write) == std::filesystem::perms::none;
+	return ((Perms & std::filesystem::perms::owner_write)  == std::filesystem::perms::none) &&
+		   ((Perms & std::filesystem::perms::group_write)  == std::filesystem::perms::none) &&
+		   ((Perms & std::filesystem::perms::others_write) == std::filesystem::perms::none);
+}
+
+inline bool
+IsExecutable(std::filesystem::perms Perms)
+{
+	return ((Perms & std::filesystem::perms::owner_exec)  != std::filesystem::perms::none) ||
+           ((Perms & std::filesystem::perms::group_exec)  != std::filesystem::perms::none) ||
+           ((Perms & std::filesystem::perms::others_exec) != std::filesystem::perms::none);
 }
 
 uint64 BlockingReadLarge(FIOReader& Reader, uint64 Offset, uint64 Size, uint8* OutputBuffer, uint64 OutputBufferSize);
@@ -487,6 +498,7 @@ FFileAttributes GetCachedFileAttrib(const FPath& Path, FFileAttributeCache& Attr
 
 bool			SetFileMtime(const FPath& Path, uint64 Mtime, bool bAllowInDryRun = false);
 bool			SetFileReadOnly(const FPath& Path, bool ReadOnly);
+bool            SetFileExecutable(const FPath& Path,  bool Executable);
 bool			IsDirectory(const FPath& Path);
 bool			PathExists(const FPath& Path);
 bool			PathExists(const FPath& Path, std::error_code& OutErrorCode);
diff --git a/Engine/Source/Programs/Unsync/Private/UnsyncHorde.cpp b/Engine/Source/Programs/Unsync/Private/UnsyncHorde.cpp
--- a/Engine/Source/Programs/Unsync/Private/UnsyncHorde.cpp
+++ b/Engine/Source/Programs/Unsync/Private/UnsyncHorde.cpp
@@ -879,6 +879,11 @@ TResult<FDirectoryManifest> DecodeHordeManifestJson(const char* JsonString, std:
 				FileManifest.bReadOnly = Field.bool_value();
 			}
 
+			if (const Json& Field = FileObject["executable"]; Field.is_bool())
+			{
+				FileManifest.bIsExecutable = Field.bool_value();
+			}
+
 			if (const Json& BlocksField = FileObject["blocks"]; BlocksField.is_array())
 			{
 				for (const Json& BlockObject : BlocksField.array_items())
diff --git a/Engine/Source/Programs/Unsync/Private/UnsyncManifest.cpp b/Engine/Source/Programs/Unsync/Private/UnsyncManifest.cpp
--- a/Engine/Source/Programs/Unsync/Private/UnsyncManifest.cpp
+++ b/Engine/Source/Programs/Unsync/Private/UnsyncManifest.cpp
@@ -315,11 +315,14 @@ CreateDirectoryManifest(const FPath& Root, const FComputeBlocksParams& Params)
 
 		FFileManifest FileManifest;
 
-		FileManifest.Mtime		 = ToWindowsFileTime(Dir.last_write_time());
-		FileManifest.Size		 = Dir.file_size();
-		FileManifest.CurrentPath = Dir.path();
-		FileManifest.BlockSize	 = Params.BlockSize;
-		FileManifest.bReadOnly	 = IsReadOnly(Dir.status().permissions());
+		std::filesystem::perms Perms = Dir.status().permissions();
+
+		FileManifest.Mtime		   = ToWindowsFileTime(Dir.last_write_time());
+		FileManifest.Size		   = Dir.file_size();
+		FileManifest.CurrentPath   = Dir.path();
+		FileManifest.BlockSize	   = Params.BlockSize;
+		FileManifest.bReadOnly	   = IsReadOnly(Perms);
+		FileManifest.bIsExecutable = IsExecutable(Perms);
 
 		{
 			std::lock_guard<std::mutex> LockGuard(ResultMutex);
diff --git a/Engine/Source/Programs/Unsync/Private/UnsyncManifest.h b/Engine/Source/Programs/Unsync/Private/UnsyncManifest.h
--- a/Engine/Source/Programs/Unsync/Private/UnsyncManifest.h
+++ b/Engine/Source/Programs/Unsync/Private/UnsyncManifest.h
@@ -25,6 +25,7 @@ struct FFileManifest
 	FPath			   CurrentPath;
 	std::string		   RevisionControlIdentity;
 	bool			   bReadOnly = false;
+	bool               bIsExecutable = false;
 
 	bool IsValid() const { return Mtime != 0 && Size != 0; }
 };
diff --git a/Engine/Source/Programs/Unsync/Private/UnsyncProtocol.h b/Engine/Source/Programs/Unsync/Private/UnsyncProtocol.h
--- a/Engine/Source/Programs/Unsync/Private/UnsyncProtocol.h
+++ b/Engine/Source/Programs/Unsync/Private/UnsyncProtocol.h
@@ -79,12 +79,14 @@ void FormatJsonBlock(std::string& Output, const FGenericBlock& Block);
 void FormatJsonBlockArray(std::wstring& Output, const FGenericBlockArray& Blocks);
 void FormatJsonBlockArray(std::string& Output, const FGenericBlockArray& Blocks);
 
+// These are just random 64 bit numbers, not based on anything in particular
 static constexpr uint64 SERIALIZED_SECTION_ID_TERMINATOR			= 0;
 static constexpr uint64 SERIALIZED_SECTION_ID_METADATA_STRING		= 0xC6BD6CDCEEF79533ull;
 static constexpr uint64 SERIALIZED_SECTION_ID_MACRO_BLOCK			= 0x8390AEBB745E08BCull;
 static constexpr uint64 SERIALIZED_SECTION_ID_FILE_READ_ONLY_MASK	= 0x851F32ED3615F0ADull;
 static constexpr uint64 SERIALIZED_SECTION_ID_FILE_REVISION_CONTROL = 0x2C1C72E6B78B1B50ull;
 static constexpr uint64 SERIALIZED_SECTION_ID_PACK_REFERENCE		= 0x634EA57F1E48DFBDull;
+static constexpr uint64 SERIALIZED_SECTION_ID_FILE_EXECUTABLE_BIT	= 0x2F4212FDAEF5C1ADull;
 
 struct FSerializedSectionHeader
 {
@@ -115,6 +117,12 @@ struct FFileReadOnlyMaskSection
 	static constexpr uint64 VERSION = 1;
 };
 
+struct FFileExecutableBitSection
+{
+	static constexpr uint64 MAGIC	= SERIALIZED_SECTION_ID_FILE_EXECUTABLE_BIT;
+	static constexpr uint64 VERSION = 1;
+};
+
 struct FFileRevisionControlSection
 {
 	static constexpr uint64 MAGIC	= SERIALIZED_SECTION_ID_FILE_REVISION_CONTROL;
diff --git a/Engine/Source/Programs/Unsync/Private/UnsyncSerialization.cpp b/Engine/Source/Programs/Unsync/Private/UnsyncSerialization.cpp
--- a/Engine/Source/Programs/Unsync/Private/UnsyncSerialization.cpp
+++ b/Engine/Source/Programs/Unsync/Private/UnsyncSerialization.cpp
@@ -290,7 +290,7 @@ LoadFileReadOnlyMask(FIOReaderStream& Reader, FSerializedSectionHeader Header)
 
 	if (Header.Version != FFileReadOnlyMaskSection::VERSION)
 	{
-		UNSYNC_ERROR(L"Found read-only file mask section version %llu, but only version %lly is supported.",
+		UNSYNC_ERROR(L"Found read-only file mask section version %llu, but only version %llu is supported.",
 					 llu(Header.Version),
 					 llu(FFileReadOnlyMaskSection::VERSION));
 		return Result;
@@ -304,6 +304,27 @@ LoadFileReadOnlyMask(FIOReaderStream& Reader, FSerializedSectionHeader Header)
 	return Result;
 }
 
+static FBuffer
+LoadFileExecutableBit(FIOReaderStream& Reader, FSerializedSectionHeader Header)
+{
+	FBuffer Result;
+
+	if (Header.Version != FFileExecutableBitSection::VERSION)
+	{
+		UNSYNC_ERROR(L"Found executable file bit section version %llu, but only version %llu is supported.",
+					 llu(Header.Version),
+					 llu(FFileExecutableBitSection::VERSION));
+		return Result;
+	}
+
+	Result.Resize(Header.Size);
+	uint8* ResultData = Result.Data();
+
+	Reader.Read(ResultData, Header.Size);
+
+	return Result;
+}
+
 static FBuffer
 SaveMacroBlocks(const FDirectoryManifest& Manifest)
 {
@@ -361,6 +382,29 @@ SaveFileReadOnlyMask(const FDirectoryManifest& Manifest)
 	return Result;
 }
 
+static FBuffer
+SaveFileExecutableBit(const FDirectoryManifest& Manifest)
+{
+	FBuffer			 Result;
+	FVectorStreamOut Writer(Result);
+
+	const uint64 NumFiles		  = Manifest.Files.size();
+	const uint64 NumBitArrayBytes = DivUp(NumFiles, 8);
+
+	Result.Resize(NumBitArrayBytes);
+	memset(Result.Data(), 0, Result.Size());
+
+	uint64 FileIndex = 0;
+	for (const auto& FileIt : Manifest.Files)
+	{
+		const FFileManifest& FileManifest = FileIt.second;
+		BitArraySet(Result.Data(), FileIndex, FileManifest.bIsExecutable);
+		++FileIndex;
+	}
+
+	return Result;
+}
+
 static FBuffer
 SaveFileRevisionControl(const FDirectoryManifest& Manifest)
 {
@@ -479,6 +523,7 @@ LoadDirectoryManifest(FDirectoryManifest& OutManifest, const FPath& Root, FIORea
 	std::unordered_map<FHash256, FGenericBlockArray> MacroBlocks;
 
 	FBuffer FileReadOnlyMask;
+	FBuffer FileExecutableMask;
 	std::vector<std::string> FileRevisions; // TODO: use linear arena to store all strings for a manifest
 
 	if (Version >= FDirectoryManifest::EVersions::V7_OptionalSections)
@@ -531,6 +576,11 @@ LoadDirectoryManifest(FDirectoryManifest& OutManifest, const FPath& Root, FIORea
 							OutManifest.PackReferences = LoadPackReferences(Stream, SectionHeader);
 							break;
 						}
+					case SERIALIZED_SECTION_ID_FILE_EXECUTABLE_BIT:
+						{
+							FileExecutableMask = LoadFileExecutableBit(Stream, SectionHeader);
+							break;
+						}
 					case SERIALIZED_SECTION_ID_TERMINATOR:
 						bDoneLoadingOptionalSections = true;
 						break;
@@ -561,6 +611,7 @@ LoadDirectoryManifest(FDirectoryManifest& OutManifest, const FPath& Root, FIORea
 		Serialize(Stream, NumFiles);
 
 		const bool bReadOnlyMaskValid = NumFiles <= FileReadOnlyMask.Size() * 8;
+		const bool bExecutableBitValid = NumFiles <= FileExecutableMask.Size() * 8;
 
 		std::string FilenameUtf8; // shared buffer to avoid some of the reallocations
 		for (uint64 FileIndex = 0; FileIndex < NumFiles; ++FileIndex)
@@ -621,6 +672,11 @@ LoadDirectoryManifest(FDirectoryManifest& OutManifest, const FPath& Root, FIORea
 				FileManifest.bReadOnly = BitArrayGet(FileReadOnlyMask.Data(), FileIndex);
 			}
 
+			if (bExecutableBitValid)
+			{
+				FileManifest.bIsExecutable = BitArrayGet(FileExecutableMask.Data(), FileIndex);
+			}
+
 			if (OutManifest.bHasFileRevisionControl)
 			{
 				FileManifest.RevisionControlIdentity = std::move(FileRevisions[FileIndex]);
@@ -770,6 +826,11 @@ SaveDirectoryManifest(const FDirectoryManifest& Manifest, FVectorStreamOut& Stre
 		WriteSection<FFileReadOnlyMaskSection>(Stream, SectionBuffer.View());
 	}
 
+	{
+		FBuffer SectionBuffer = SaveFileExecutableBit(Manifest);
+		WriteSection<FFileExecutableBitSection>(Stream, SectionBuffer.View());
+	}
+
 	if (Manifest.bHasFileRevisionControl)
 	{
 		FBuffer SectionBuffer = SaveFileRevisionControl(Manifest);
diff --git a/Engine/Source/Programs/Unsync/Private/UnsyncVersion.cpp b/Engine/Source/Programs/Unsync/Private/UnsyncVersion.cpp
--- a/Engine/Source/Programs/Unsync/Private/UnsyncVersion.cpp
+++ b/Engine/Source/Programs/Unsync/Private/UnsyncVersion.cpp
@@ -2,7 +2,7 @@
 
 #include "UnsyncVersion.h"
 
-#define UNSYNC_VERSION_STR "1.0.91"
+#define UNSYNC_VERSION_STR "1.0.92"
 
 namespace unsync {
 
