PR #12563: Handle TLS_WANT_POLLIN/POLLOUT error codes during TLS IO

diff --git a/Engine/Source/Programs/Unsync/Private/UnsyncSocket.cpp b/Engine/Source/Programs/Unsync/Private/UnsyncSocket.cpp
--- a/Engine/Source/Programs/Unsync/Private/UnsyncSocket.cpp
+++ b/Engine/Source/Programs/Unsync/Private/UnsyncSocket.cpp
@@ -434,7 +434,14 @@ FSocketTls::FSocketTls(FSocketHandle InHandle, FTlsClientSettings ClientSettings
 
 	if (Err == 0)
 	{
-		Err = tls_handshake(TlsCtx);
+		for(;;)
+		{
+			Err = tls_handshake(TlsCtx);
+			if (Err != TLS_WANT_POLLIN && Err != TLS_WANT_POLLOUT)
+			{
+				break;
+			}
+		}
 	}
 
 	if (Err)
@@ -472,11 +479,31 @@ FSocketTls::~FSocketTls()
 {
 	if (TlsCtx)
 	{
-		tls_close(TlsCtx);
+		for(;;)
+		{
+			int Err = tls_close(TlsCtx);
+			if (Err != TLS_WANT_POLLIN && Err != TLS_WANT_POLLOUT)
+			{
+				break;
+			}
+		}
 		tls_free(TlsCtx);
 	}
 }
 
+static ssize_t 
+TlsWrite(struct tls* TlsCtx, const char* Buffer, size_t BufferSize)
+{
+	ssize_t ErrorCode;
+	do
+	{
+		ErrorCode = tls_write(TlsCtx, Buffer, BufferSize);
+	} 
+	while(ErrorCode == TLS_WANT_POLLIN || ErrorCode == TLS_WANT_POLLOUT);
+
+	return ErrorCode;
+}
+
 int
 FSocketTls::Send(const void* Data, size_t DataSize)
 {
@@ -488,7 +515,7 @@ FSocketTls::Send(const void* Data, size_t DataSize)
 	int ProcessedBytes = 0;
 	while (ProcessedBytes < DataSize)
 	{
-		ssize_t ErrorCode = tls_write(TlsCtx, (const char*)Data + ProcessedBytes, DataSize - ProcessedBytes);
+		ssize_t ErrorCode = TlsWrite(TlsCtx, (const char*)Data + ProcessedBytes, DataSize - ProcessedBytes);
 		if (ErrorCode <= 0)
 		{
 			if (ErrorCode < 0)
@@ -504,6 +531,19 @@ FSocketTls::Send(const void* Data, size_t DataSize)
 	return ProcessedBytes;
 }
 
+static ssize_t 
+TlsRead(struct tls* TlsCtx, char* Buffer, size_t BufferSize)
+{
+	ssize_t ErrorCode;
+	do
+	{
+		ErrorCode = tls_read(TlsCtx, Buffer, BufferSize);
+	} 
+	while(ErrorCode == TLS_WANT_POLLIN || ErrorCode == TLS_WANT_POLLOUT);
+
+	return ErrorCode;
+}
+
 int
 FSocketTls::RecvAll(void* Data, size_t DataSize)
 {
@@ -515,7 +555,7 @@ FSocketTls::RecvAll(void* Data, size_t DataSize)
 	int ProcessedBytes = 0;
 	while (ProcessedBytes < DataSize)
 	{
-		ssize_t ErrorCode = tls_read(TlsCtx, (char*)Data + ProcessedBytes, DataSize - ProcessedBytes);
+		ssize_t ErrorCode = TlsRead(TlsCtx, (char*)Data + ProcessedBytes, DataSize - ProcessedBytes);
 		if (ErrorCode <= 0)
 		{
 			if (ErrorCode < 0)
@@ -542,7 +582,7 @@ FSocketTls::RecvAny(void* Data, size_t DataSize)
 	int ProcessedBytes = 0;
 	do
 	{
-		ssize_t ErrorCode = tls_read(TlsCtx, (char*)Data + ProcessedBytes, DataSize - ProcessedBytes);
+		ssize_t ErrorCode = TlsRead(TlsCtx, (char*)Data + ProcessedBytes, DataSize - ProcessedBytes);
 		if (ErrorCode <= 0)
 		{
 			if (ErrorCode < 0)
