unsync - Fix block index file loading during pack command

* Use the common helper function to load the data

diff --git a/Engine/Source/Programs/Unsync/Private/UnsyncCmdPack.cpp b/Engine/Source/Programs/Unsync/Private/UnsyncCmdPack.cpp
--- a/Engine/Source/Programs/Unsync/Private/UnsyncCmdPack.cpp
+++ b/Engine/Source/Programs/Unsync/Private/UnsyncCmdPack.cpp
@@ -313,6 +313,9 @@ CmdPack(const FCmdPackOptions& Options)
 	{
 		UNSYNC_LOG_INDENT;
 		FPath ExpectedExtension = FPath(".unsync_index");
+
+		FPackIndexDatabase TempPackIndexDatabase;
+
 		for (const std::filesystem::directory_entry& Dir : RecursiveDirectoryScan(PackRoot))
 		{
 			if (!Dir.is_regular_file())
@@ -327,9 +330,21 @@ CmdPack(const FCmdPackOptions& Options)
 				continue;
 			}
 
-			FBuffer ExistingEntries = ReadFileToBuffer(FilePath);
-			uint64	NumEntries		= ExistingEntries.Size() / sizeof(FPackIndexEntry);
-			for (const FPackIndexEntry& Entry : MakeView(reinterpret_cast<FPackIndexEntry*>(ExistingEntries.Data()), NumEntries))
+			FNativeFile IndexFile(FilePath, EFileMode::ReadOnly);
+			if (!IndexFile.IsValid())
+			{
+				UNSYNC_WARNING(L"Failed to open pack index file '%ls'", FilePath.wstring().c_str());
+				continue;
+			}
+
+			FIOReaderStream IndexFileStream(IndexFile);
+			if (!LoadPackIndexDatabase(TempPackIndexDatabase, IndexFileStream))
+			{
+				UNSYNC_WARNING(L"Failed to read pack index database from file '%ls'", FilePath.wstring().c_str());
+				continue;
+			}
+
+			for (const FPackIndexEntry& Entry : TempPackIndexDatabase.Entries)
 			{
 				SeenBlockHashSet.insert(Entry.BlockHash);
 			}
