unsync - Remove FIOBase::FlushOne since it was unused

diff --git a/Engine/Source/Programs/Unsync/Private/UnsyncFile.cpp b/Engine/Source/Programs/Unsync/Private/UnsyncFile.cpp
--- a/Engine/Source/Programs/Unsync/Private/UnsyncFile.cpp
+++ b/Engine/Source/Programs/Unsync/Private/UnsyncFile.cpp
@@ -736,58 +736,6 @@ FWindowsFile::FlushAll()
 	InternalFlushAll();
 }
 
-void
-FWindowsFile::FlushOne()
-{
-	std::lock_guard<std::mutex> LockGuard(Mutex);
-
-	uint32 NumValidEvents			  = 0;
-	HANDLE Events[NUM_QUEUES]		  = {};
-	uint32 CommandIndices[NUM_QUEUES] = {};
-
-	// todo: maintain a list of outstanding commands in read()
-	for (uint32 I = 0; I < NUM_QUEUES; ++I)
-	{
-		if (Commands[I].Buffer.GetMemory())
-		{
-			Events[NumValidEvents]		   = OverlappedEvents[I];
-			CommandIndices[NumValidEvents] = I;
-			NumValidEvents++;
-		}
-	}
-
-	if (NumValidEvents == 0)
-	{
-		return;
-	}
-
-	DWORD WaitResult = WaitForMultipleObjects(NumValidEvents, OverlappedEvents, false, INFINITE);
-
-	uint32 CmdIdx = ~0u;
-	if (WaitResult >= WAIT_OBJECT_0 && WaitResult < NUM_QUEUES)
-	{
-		CmdIdx = CommandIndices[WaitResult - WAIT_OBJECT_0];
-	}
-
-	UNSYNC_ASSERT(CmdIdx < NUM_QUEUES);
-	FWindowsFile::Command& Cmd = Commands[CmdIdx];
-
-	uint32 ReadBytes = CompleteReadCommand(Cmd);
-
-	if (ReadBytes != Cmd.ReadSize)
-	{
-		return;
-	}
-
-	UNSYNC_ASSERT(Cmd.Buffer.GetMemory());
-	UNSYNC_ASSERT(ReadBytes);
-
-	Cmd.ReadSize = 0;
-	Cmd.UserData = 0;
-	Cmd.Buffer	 = {};
-	Cmd.Callback = {};
-}
-
 FFileAttributes
 GetFileAttrib(const FPath& Path, FFileAttributeCache* AttribCache)
 {
@@ -883,12 +831,6 @@ FUnixFile::FlushAll()
 	// nothing to do, all file IO is synchronous
 }
 
-void
-FUnixFile::FlushOne()
-{
-	// nothing to do, all file IO is synchronous
-}
-
 void
 FUnixFile::Close()
 {
diff --git a/Engine/Source/Programs/Unsync/Private/UnsyncFile.h b/Engine/Source/Programs/Unsync/Private/UnsyncFile.h
--- a/Engine/Source/Programs/Unsync/Private/UnsyncFile.h
+++ b/Engine/Source/Programs/Unsync/Private/UnsyncFile.h
@@ -131,7 +131,6 @@ struct FIOBase
 {
 	virtual ~FIOBase()		  = default;
 	virtual void   FlushAll() = 0;
-	virtual void   FlushOne() = 0;
 	virtual uint64 GetSize()  = 0;
 	virtual bool   IsValid()  = 0;
 	virtual void   Close()	  = 0;
@@ -174,7 +173,6 @@ struct FWindowsFile : FIOReaderWriter
 
 	// IOBase
 	virtual void   FlushAll() override;
-	virtual void   FlushOne() override;
 	virtual uint64 GetSize() override { return FileSize; }
 	virtual bool   IsValid() override;
 	virtual void   Close() override;
@@ -232,7 +230,6 @@ struct FUnixFile : FIOReaderWriter
 
 	// IOBase
 	virtual void   FlushAll() override;
-	virtual void   FlushOne() override;
 	virtual uint64 GetSize() override { return FileSize; }
 	virtual bool   IsValid() override { return FileHandle != nullptr; }
 	virtual void   Close() override;
@@ -296,7 +293,6 @@ struct FMemReader : FIOReader
 
 	// IOBase
 	virtual void   FlushAll() override {}
-	virtual void   FlushOne() override {}
 	virtual uint64 GetSize() override { return Size; }
 	virtual bool   IsValid() override { return Data != nullptr; }
 	virtual void   Close() override
@@ -354,11 +350,10 @@ struct FNullReaderWriter : FIOReaderWriter
 	explicit FNullReaderWriter(FInvalid) : DataSize(0), bValid(false) {}
 
 	// IOBase
-	virtual void   FlushAll() override{};
-	virtual void   FlushOne() override{};
+	virtual void   FlushAll() override {};
 	virtual uint64 GetSize() override { return DataSize; }
 	virtual bool   IsValid() override { return bValid; }
-	virtual void   Close() override{};
+	virtual void   Close() override {};
 	virtual int32  GetError() override { return 0; }
 
 	// IORead
@@ -390,7 +385,6 @@ struct FDeferredOpenReader : FIOReader
 
 	// IOBase
 	virtual void   FlushAll() override { GetOrOpenInner()->FlushAll(); }
-	virtual void   FlushOne() override { GetOrOpenInner()->FlushOne(); }
 	virtual uint64 GetSize() override { return GetOrOpenInner()->GetSize(); }
 	virtual bool   IsValid() override { return GetOrOpenInner()->IsValid(); }
 	virtual void   Close() override { if (Inner) { Inner->Close(); } }
