unsync - Add FBlake3Hasher wrapper and use it instead of raw blake3_hasher API

diff --git a/Engine/Source/Programs/Unsync/Private/UnsyncChunking.cpp b/Engine/Source/Programs/Unsync/Private/UnsyncChunking.cpp
--- a/Engine/Source/Programs/Unsync/Private/UnsyncChunking.cpp
+++ b/Engine/Source/Programs/Unsync/Private/UnsyncChunking.cpp
@@ -7,7 +7,6 @@
 #include "UnsyncScheduler.h"
 
 UNSYNC_THIRD_PARTY_INCLUDES_START
-#include <blake3.h>
 #include <md5-sse2.h>
 UNSYNC_THIRD_PARTY_INCLUDES_END
 
@@ -118,8 +117,7 @@ ComputeBlocksVariableT(FIOReader& Reader, const FComputeBlocksParams& Params)
 			CurrentMacroBlock.HashStrong.Type = EHashType::Blake3_160;
 			CurrentMacroBlock.Offset		  = Task.Offset;
 
-			blake3_hasher MacroBlockHasher;
-			blake3_hasher_init(&MacroBlockHasher);
+			FBlake3Hasher MacroBlockHasher;
 
 			auto ScanFn = [MacroBlockHashThreshold,
 						   MaximumMacroBlockSize,
@@ -149,7 +147,7 @@ ComputeBlocksVariableT(FIOReader& Reader, const FComputeBlocksParams& Params)
 
 					if (TargetMacroBlockSize)
 					{
-						blake3_hasher_update(&MacroBlockHasher, LastBlockEnd, ThisBlockSize);
+						MacroBlockHasher.Update(LastBlockEnd, ThisBlockSize);
 						CurrentMacroBlock.Size += Block.Size;
 
 						uint32 HashStrong32 = 0;
@@ -159,9 +157,9 @@ ComputeBlocksVariableT(FIOReader& Reader, const FComputeBlocksParams& Params)
 							(CurrentMacroBlock.Size + Block.Size > MaximumMacroBlockSize) || bLastBlock)
 						{
 							// Commit the macro block
-							blake3_hasher_finalize(&MacroBlockHasher,
-												   CurrentMacroBlock.HashStrong.Data,
-												   sizeof(CurrentMacroBlock.HashStrong.Data));
+							const FHash256 MacroBlockHash = MacroBlockHasher.Finalize();
+							CurrentMacroBlock.HashStrong = FGenericHash::FromBlake3_256(MacroBlockHash);
+
 							Task.MacroBlocks.push_back(CurrentMacroBlock);
 
 							if (Params.OnMacroBlockGenerated)
@@ -173,7 +171,7 @@ ComputeBlocksVariableT(FIOReader& Reader, const FComputeBlocksParams& Params)
 							}
 
 							// Reset macro block state
-							blake3_hasher_init(&MacroBlockHasher);
+							MacroBlockHasher.Reset();
 							CurrentMacroBlock.Offset += CurrentMacroBlock.Size;
 							CurrentMacroBlock.Size = 0;
 						}
diff --git a/Engine/Source/Programs/Unsync/Private/UnsyncHash.cpp b/Engine/Source/Programs/Unsync/Private/UnsyncHash.cpp
--- a/Engine/Source/Programs/Unsync/Private/UnsyncHash.cpp
+++ b/Engine/Source/Programs/Unsync/Private/UnsyncHash.cpp
@@ -24,6 +24,39 @@ HashBlake3Bytes(const uint8* Data, uint64 Size)
 	return Result;
 }
 
+FBlake3Hasher::FBlake3Hasher()
+{
+	static_assert(sizeof(State) >= sizeof(blake3_hasher));
+	static_assert(sizeof(State) == sizeof(FBlake3Hasher));
+	static_assert(alignof(FBlake3Hasher) >= alignof(blake3_hasher));
+
+	blake3_hasher* Hasher = reinterpret_cast<blake3_hasher*>(State);
+	blake3_hasher_init(Hasher);
+}
+
+void
+FBlake3Hasher::Update(const void* Data, size_t Size)
+{
+	blake3_hasher* Hasher = reinterpret_cast<blake3_hasher*>(State);
+	blake3_hasher_update(Hasher, Data, Size);
+}
+
+FHash256
+FBlake3Hasher::Finalize()
+{
+	FHash256 Result;
+	blake3_hasher* Hasher = reinterpret_cast<blake3_hasher*>(State);
+	blake3_hasher_finalize(Hasher, Result.Data, sizeof(Result.Data));
+	return Result;
+}
+
+void
+FBlake3Hasher::Reset()
+{
+	blake3_hasher* Hasher = reinterpret_cast<blake3_hasher*>(State);
+	blake3_hasher_reset(Hasher);
+}
+
 template FHash128 HashBlake3Bytes<FHash128>(const uint8* data, uint64 size);
 template FHash160 HashBlake3Bytes<FHash160>(const uint8* data, uint64 size);
 template FHash256 HashBlake3Bytes<FHash256>(const uint8* data, uint64 size);
diff --git a/Engine/Source/Programs/Unsync/Private/UnsyncHash.h b/Engine/Source/Programs/Unsync/Private/UnsyncHash.h
--- a/Engine/Source/Programs/Unsync/Private/UnsyncHash.h
+++ b/Engine/Source/Programs/Unsync/Private/UnsyncHash.h
@@ -298,6 +298,20 @@ HashBlake3String(const std::wstring& Str)
 	return HashBlake3Bytes<HashResultType>((const uint8*)Str.c_str(), Str.length() * sizeof(Str[0]));
 }
 
+class alignas(8) FBlake3Hasher
+{
+public:
+	FBlake3Hasher();
+
+	void	 Update(const void* Data, size_t Size);
+	FHash256 Finalize();
+	void	 Reset();
+
+private:
+	static constexpr size_t STATE_SIZE = 1912;
+	uint8 State[STATE_SIZE];
+};
+
 FHash128 HashMd5Bytes(const uint8* Data, uint64 Size);
 
 FGenericHash HashBytes(const uint8* Data, uint64 Size, EHashType HashType);
diff --git a/Engine/Source/Programs/Unsync/Private/UnsyncManifest.cpp b/Engine/Source/Programs/Unsync/Private/UnsyncManifest.cpp
--- a/Engine/Source/Programs/Unsync/Private/UnsyncManifest.cpp
+++ b/Engine/Source/Programs/Unsync/Private/UnsyncManifest.cpp
@@ -11,10 +11,6 @@
 #include "UnsyncSerialization.h"
 #include "UnsyncThread.h"
 
-UNSYNC_THIRD_PARTY_INCLUDES_START
-#include <blake3.h>
-UNSYNC_THIRD_PARTY_INCLUDES_END
-
 namespace unsync {
 
 FDirectoryManifestInfo
@@ -115,28 +111,27 @@ LogManifestFiles(ELogLevel LogLevel, const FDirectoryManifest& Manifest)
 
 template<typename T>
 inline auto
-UpdateHashT(blake3_hasher& Hasher, const T& V)
+UpdateHashT(FBlake3Hasher& Hasher, const T& V)
 {
-	blake3_hasher_update(&Hasher, &V, sizeof(T));
+	Hasher.Update(&V, sizeof(T));
 }
 
 void
-UpdateHashBlocks(blake3_hasher& Hasher, const FGenericBlockArray& Blocks)
+UpdateHashBlocks(FBlake3Hasher& Hasher, const FGenericBlockArray& Blocks)
 {
 	for (const FGenericBlock& Block : Blocks)
 	{
 		UpdateHashT(Hasher, Block.Offset);
 		UpdateHashT(Hasher, Block.Size);
 		UpdateHashT(Hasher, Block.HashWeak);
-		blake3_hasher_update(&Hasher, Block.HashStrong.Data, Block.HashStrong.Size());
+		Hasher.Update(Block.HashStrong.Data, Block.HashStrong.Size());
 	}
 }
 
 FHash256
 ComputeManifestStableSignature(const FDirectoryManifest& Manifest)
 {
-	blake3_hasher Hasher;
-	blake3_hasher_init(&Hasher);
+	FBlake3Hasher Hasher;
 
 	UpdateHashT(Hasher, Manifest.Algorithm.ChunkingAlgorithmId);
 	UpdateHashT(Hasher, Manifest.Algorithm.WeakHashAlgorithmId);
@@ -152,7 +147,7 @@ ComputeManifestStableSignature(const FDirectoryManifest& Manifest)
 		ConvertWideToUtf8(FileName, FileNameUtf8);
 		std::replace(FileNameUtf8.begin(), FileNameUtf8.end(), '\\', '/');
 
-		blake3_hasher_update(&Hasher, FileNameUtf8.c_str(), FileNameUtf8.length());
+		Hasher.Update(FileNameUtf8.c_str(), FileNameUtf8.length());
 
 		UpdateHashT(Hasher, FileManifest.Mtime);
 		UpdateHashT(Hasher, FileManifest.Size);
@@ -162,8 +157,7 @@ ComputeManifestStableSignature(const FDirectoryManifest& Manifest)
 		UpdateHashBlocks(Hasher, FileManifest.MacroBlocks);
 	}
 
-	FHash256 Result;
-	blake3_hasher_finalize(&Hasher, Result.Data, sizeof(Result.Data));
+	FHash256 Result = Hasher.Finalize();
 
 	return Result;
 }
