unsync - Add global declaration of macro block type

diff --git a/Engine/Source/Programs/Unsync/Private/UnsyncChunking.cpp b/Engine/Source/Programs/Unsync/Private/UnsyncChunking.cpp
--- a/Engine/Source/Programs/Unsync/Private/UnsyncChunking.cpp
+++ b/Engine/Source/Programs/Unsync/Private/UnsyncChunking.cpp
@@ -114,7 +114,7 @@ ComputeBlocksVariableT(FIOReader& Reader, const FComputeBlocksParams& Params)
 			const uint8* LastBlockEnd = DataBegin;
 
 			FGenericBlock CurrentMacroBlock;
-			CurrentMacroBlock.HashStrong.Type = EHashType::Blake3_160;
+			CurrentMacroBlock.HashStrong.Type = MACRO_BLOCK_HASH_TYPE;
 			CurrentMacroBlock.Offset		  = Task.Offset;
 
 			FBlake3Hasher MacroBlockHasher;
@@ -157,8 +157,9 @@ ComputeBlocksVariableT(FIOReader& Reader, const FComputeBlocksParams& Params)
 							(CurrentMacroBlock.Size + Block.Size > MaximumMacroBlockSize) || bLastBlock)
 						{
 							// Commit the macro block
-							const FHash256 MacroBlockHash = MacroBlockHasher.Finalize();
-							CurrentMacroBlock.HashStrong = FGenericHash::FromBlake3_256(MacroBlockHash);
+							const FHash256 MacroBlockHash	  = MacroBlockHasher.Finalize();
+							CurrentMacroBlock.HashStrong	  = FGenericHash::FromBlake3_256(MacroBlockHash);
+							CurrentMacroBlock.HashStrong.Type = MACRO_BLOCK_HASH_TYPE;
 
 							Task.MacroBlocks.push_back(CurrentMacroBlock);
 
diff --git a/Engine/Source/Programs/Unsync/Private/UnsyncHash.h b/Engine/Source/Programs/Unsync/Private/UnsyncHash.h
--- a/Engine/Source/Programs/Unsync/Private/UnsyncHash.h
+++ b/Engine/Source/Programs/Unsync/Private/UnsyncHash.h
@@ -49,7 +49,7 @@ enum class EHashType : uint8 {
 	Unknown_256,
 	MD5,		 // Standard 128 bit MD5
 	Blake3_128,	 // Blake3 hash, truncated to 128 bits
-	Blake3_160,	 // Blake3 hash, truncated to 128 bits, AKA Unreal Engine IoHash
+	Blake3_160,	 // Blake3 hash, truncated to 160 bits, AKA Unreal Engine IoHash
 	Blake3_256,	 // Blake3 hash, full 256 bits
 };
 
diff --git a/Engine/Source/Programs/Unsync/Private/UnsyncJupiter.cpp b/Engine/Source/Programs/Unsync/Private/UnsyncJupiter.cpp
--- a/Engine/Source/Programs/Unsync/Private/UnsyncJupiter.cpp
+++ b/Engine/Source/Programs/Unsync/Private/UnsyncJupiter.cpp
@@ -580,7 +580,7 @@ JupiterPush(const FDirectoryManifest& Manifest, const FRemoteDesc& RemoteDesc, c
 						}
 
 						// TODO: pipeline PUT commands
-						UNSYNC_ASSERT(Block.HashStrong.Type == EHashType::Blake3_160);
+						UNSYNC_ASSERT(Block.HashStrong.Type == MACRO_BLOCK_HASH_TYPE);
 						if (jupiter_put_raw_blob(*Connection,
 												 RemoteDesc.StorageNamespace,
 												 RemoteDesc.HttpHeaders,
@@ -978,7 +978,7 @@ FJupiterProtocolImpl::Download(const TArrayView<FNeedBlock> NeedBlocks, const FB
 														 MacroBlockRequest.Offset + MacroBlockRequest.Size - 1,
 														 ToStringView(HttpHeaders));
 
-				UNSYNC_ASSERT(MacroBlockRequest.Hash.Type == EHashType::Blake3_160);
+				UNSYNC_ASSERT(MacroBlockRequest.Hash.Type == MACRO_BLOCK_HASH_TYPE);
 				bool bRequestAdded = JupiterGetRawBlobBegin(Connection,
 															RemoteDesc.StorageNamespace,
 															ToStringView(RequestHeaders),
diff --git a/Engine/Source/Programs/Unsync/Private/UnsyncProtocol.h b/Engine/Source/Programs/Unsync/Private/UnsyncProtocol.h
--- a/Engine/Source/Programs/Unsync/Private/UnsyncProtocol.h
+++ b/Engine/Source/Programs/Unsync/Private/UnsyncProtocol.h
@@ -18,7 +18,8 @@
 
 namespace unsync {
 
-static constexpr uint32 MAX_BLOCK_SIZE = uint32(1_MB);
+static constexpr uint32	   MAX_BLOCK_SIZE		 = uint32(1_MB);
+static constexpr EHashType MACRO_BLOCK_HASH_TYPE = EHashType::Blake3_160;  // Blake3 160bit (IoHash) is required for macro blocks due to back-end storage implementation
 
 enum class EChunkingAlgorithmID : uint64
 {
diff --git a/Engine/Source/Programs/Unsync/Private/UnsyncSerialization.cpp b/Engine/Source/Programs/Unsync/Private/UnsyncSerialization.cpp
--- a/Engine/Source/Programs/Unsync/Private/UnsyncSerialization.cpp
+++ b/Engine/Source/Programs/Unsync/Private/UnsyncSerialization.cpp
@@ -252,7 +252,7 @@ LoadMacroBlocks(FIOReaderStream& Reader, uint64 Version)
 	uint64 NumFiles = 0;
 	Serialize(Reader, NumFiles);
 
-	const EHashType MacroBlockHashType = EHashType::Blake3_160;
+	const EHashType MacroBlockHashType = MACRO_BLOCK_HASH_TYPE;
 	const uint32	MacroBlockHashSize = uint32(GetHashSize(MacroBlockHashType));
 
 	for (uint64 FileIndex = 0; FileIndex < NumFiles; ++FileIndex)
@@ -331,7 +331,7 @@ SaveMacroBlocks(const FDirectoryManifest& Manifest)
 	FBuffer			 Result;
 	FVectorStreamOut Writer(Result);
 
-	const EHashType MacroBlockHashType = EHashType::Blake3_160;
+	const EHashType MacroBlockHashType = MACRO_BLOCK_HASH_TYPE;
 	const uint32	MacroBlockHashSize = uint32(GetHashSize(MacroBlockHashType));
 
 	Writer.WriteT<uint64>(Manifest.Files.size());
