unsync - Add ordered async file reading mechanism

This can be used in the future to implement a more memory-efficient manifest generation / chunking scheme

diff --git a/Engine/Source/Programs/Unsync/Private/UnsyncFile.cpp b/Engine/Source/Programs/Unsync/Private/UnsyncFile.cpp
--- a/Engine/Source/Programs/Unsync/Private/UnsyncFile.cpp
+++ b/Engine/Source/Programs/Unsync/Private/UnsyncFile.cpp
@@ -6,6 +6,7 @@
 #include "UnsyncThread.h"
 #include "UnsyncScheduler.h"
 #include "UnsyncFilter.h"
+#include "UnsyncHash.h"
 
 #include <mutex>
 
@@ -370,7 +371,7 @@ FWindowsFile::FWindowsFile(const FPath& InFilename, EFileMode InMode, uint64 InS
 			UNSYNC_ERROR(L"Unexpected file mode %d", (int)InMode);
 		}
 
-		for (uint32 I = 0; I < NUM_QUEUES; ++I)
+		for (uint32 I = 0; I < MaxQueueDepth; ++I)
 		{
 			OverlappedEvents[I]			  = CreateEvent(nullptr, true, true, nullptr);
 			Commands[I].Overlapped.hEvent = OverlappedEvents[I];
@@ -422,7 +423,7 @@ FWindowsFile::Close()
 		CloseHandle(FileHandle);
 		FileHandle = INVALID_HANDLE_VALUE;
 	}
-	for (uint32 I = 0; I < NUM_QUEUES; ++I)
+	for (uint32 I = 0; I < MaxQueueDepth; ++I)
 	{
 		if (OverlappedEvents[I])
 		{
@@ -433,7 +434,7 @@ FWindowsFile::Close()
 }
 
 uint32
-FWindowsFile::CompleteReadCommand(Command& Cmd)
+FWindowsFile::CompleteReadCommand(FOverlappedCommand& Cmd)
 {
 	// Expects that Mutex is locked
 
@@ -473,7 +474,7 @@ FWindowsFile::CompleteReadCommand(Command& Cmd)
 			SchedulerSleep(1000);
 
 			// Cancel any active overlapped IO commands
-			for (Command& It : Commands)
+			for (FOverlappedCommand& It : Commands)
 			{
 				if (It.bActive && (&It != &Cmd))
 				{
@@ -512,7 +513,7 @@ FWindowsFile::CompleteReadCommand(Command& Cmd)
 	if (bRecoveredFromError)
 	{
 		// Re-issue the read requests using new file handle
-		for (Command& It : Commands)
+		for (FOverlappedCommand& It : Commands)
 		{
 			if (It.bActive)
 			{
@@ -657,17 +658,29 @@ FWindowsFile::ReadAsync(uint64 SourceOffset, uint64 Size, uint64 UserData, IOCal
 
 	uint32 CmdIdx = ~0u;
 
-	// find available queue or wait for one to complete
-	DWORD WaitResult = WaitForMultipleObjects(NUM_QUEUES, OverlappedEvents, false, INFINITE);
+	const bool bIsOrdered = !!(Mode & EFileMode::Ordered);
 
-	if (WaitResult >= WAIT_OBJECT_0 && WaitResult < NUM_QUEUES)
+	if (bIsOrdered)
 	{
-		CmdIdx = WaitResult - WAIT_OBJECT_0;
+		const uint32 WaitSlotIndex = uint32(NumCommandsIssued % MaxQueueDepth);
+		const DWORD WaitResult = WaitForSingleObject(OverlappedEvents[WaitSlotIndex], INFINITE);
+		UNSYNC_ASSERT(WaitResult == WAIT_OBJECT_0);
+		CmdIdx = WaitSlotIndex;
+	}
+	else
+	{
+		// find available queue or wait for one to complete
+		const DWORD WaitResult = WaitForMultipleObjects(MaxQueueDepth, OverlappedEvents, false, INFINITE);
+
+		if (WaitResult >= WAIT_OBJECT_0 && WaitResult < MaxQueueDepth)
+		{
+			CmdIdx = WaitResult - WAIT_OBJECT_0;
+		}
 	}
 
-	UNSYNC_ASSERT(CmdIdx < NUM_QUEUES);
+	UNSYNC_ASSERT(CmdIdx < MaxQueueDepth);
 
-	FWindowsFile::Command& Cmd = Commands[CmdIdx];
+	FWindowsFile::FOverlappedCommand& Cmd = Commands[CmdIdx];
 
 	if (Cmd.bActive)
 	{
@@ -713,17 +726,20 @@ FWindowsFile::ReadAsync(uint64 SourceOffset, uint64 Size, uint64 UserData, IOCal
 	}
 	Cmd.bActive = true;
 
+	++NumCommandsIssued;
+
 	return true;
 }
 
 void
 FWindowsFile::InternalFlushAll()
 {
-	for (uint32 I = 0; I < NUM_QUEUES; ++I)
+	for (uint32 I = 0; I < MaxQueueDepth; ++I)
 	{
-		if (Commands[I].bActive)
+		const uint64 CommandIndex = (NumCommandsIssued + I) % MaxQueueDepth;
+		if (Commands[CommandIndex].bActive)
 		{
-			CompleteReadCommand(Commands[I]);
+			CompleteReadCommand(Commands[CommandIndex]);
 		}
 	}
 }
@@ -1051,6 +1067,12 @@ WriteBufferToFile(const FPath& Filename, const uint8* Data, uint64 Size, EFileMo
 	if (File.IsValid())
 	{
 		uint64 WroteBytes = File.Write(Data, 0, Size);
+		if (WroteBytes != Size)
+		{
+			UNSYNC_ERROR(L"Failed to write complete file '%ls'. Expected to write %llu bytes, actually written %llu bytes",
+						 Filename.wstring().c_str(),
+						 llu(Size) , llu(WroteBytes));
+		}
 		return WroteBytes == Size;
 	}
 	else
@@ -1776,4 +1798,141 @@ SystemTimeFromFileTime(std::filesystem::file_time_type FileTime)
 #endif // 0
 }
 
+void
+TestFileAsyncRead()
+{
+	UNSYNC_LOG(L"TestFileAsyncRead()");
+	UNSYNC_LOG_INDENT;
+
+	UNSYNC_LOG(L"Initializing test data");
+
+	const FPath TempDirPath = std::filesystem::temp_directory_path() / "unsync_test";
+	CreateDirectories(TempDirPath);
+
+	const bool bDirectoryExists = PathExists(TempDirPath) && IsDirectory(TempDirPath);
+	UNSYNC_ASSERT(bDirectoryExists);
+
+	const FPath		 TestFilename = TempDirPath / "ordered_integers.bin";
+	constexpr uint64 TestFileSize = 1_GB;
+
+	FBuffer TempBuffer;
+	TempBuffer.Resize(TestFileSize);
+	uint32* TempBufferData = reinterpret_cast<uint32*>(TempBuffer.Data());
+
+	const uint32 NumElements = uint32(TestFileSize / sizeof(uint32));
+
+	for (uint32 i = 0; i < NumElements; ++i)
+	{
+		TempBufferData[i] = i;
+	}
+
+	const FHash256 ExpectedHash = HashBlake3Bytes<FHash256>(reinterpret_cast<const uint8*>(TempBufferData), TestFileSize);
+
+	if (!PathExists(TestFilename))
+	{
+		UNSYNC_LOG(L"Writing test file '%ls'", TestFilename.wstring().c_str());
+		const bool bFileWritten = WriteBufferToFile(TestFilename, TempBuffer);
+		UNSYNC_ASSERT(bFileWritten);
+	}
+
+	UNSYNC_LOG(L"ReadOnlyUnbuffered");
+
+	{
+		UNSYNC_LOG_INDENT;
+
+		memset(TempBufferData, 0, TestFileSize);
+
+		FNativeFile		 TestFile(TestFilename, EFileMode::ReadOnlyUnbuffered);
+		constexpr uint64 ChunkSize = 8_MB;
+		static_assert(TestFileSize % ChunkSize == 0);
+
+		UNSYNC_LOG(L"Reading test data");
+
+		const FTimePoint ReadStartTime = TimePointNow();
+
+		const uint64 NumChunks = TestFileSize / ChunkSize;
+
+		auto IoCallback = [TempBufferData](FIOBuffer Buffer, uint64 SourceOffset, uint64 ReadSize, uint64 UserData)
+		{ memcpy(reinterpret_cast<uint8*>(TempBufferData) + SourceOffset, Buffer.GetData(), Buffer.GetSize()); };
+
+		for (uint32 ChunkIndex = 0; ChunkIndex < NumChunks; ++ChunkIndex)
+		{
+			uint64 ChunkOffset = ChunkIndex * ChunkSize;
+
+			TestFile.ReadAsync(ChunkOffset, ChunkSize, 0, IoCallback);
+		}
+
+		TestFile.FlushAll();
+
+		const FTimePoint ReadDoneTime = TimePointNow();
+
+		UNSYNC_LOG(L"Hashing test data");
+
+		const FHash256 ActualHash = HashBlake3Bytes<FHash256>(reinterpret_cast<const uint8*>(TempBufferData), TestFileSize);
+
+		const FTimePoint HashDoneTime = TimePointNow();
+
+		const double ReadDuration = DurationSec(ReadStartTime, ReadDoneTime);
+		const double HashDuration = DurationSec(ReadDoneTime, HashDoneTime);
+		const double TotalDuration = DurationSec(ReadStartTime, HashDoneTime);
+
+		UNSYNC_LOG(L"Read rate: %.2f MB / sec", SizeMb(TestFileSize) / ReadDuration);
+		UNSYNC_LOG(L"Hash rate: %.2f MB / sec", SizeMb(TestFileSize) / HashDuration);
+		UNSYNC_LOG(L"Total rate: %.2f MB / sec", SizeMb(TestFileSize) / TotalDuration);
+		UNSYNC_LOG(L"Total time: %.3f sec", TotalDuration);
+
+		UNSYNC_ASSERT(ActualHash == ExpectedHash);
+	}
+
+	UNSYNC_LOG(L"ReadOnlyUnbufferedOrdered");
+
+	{
+		UNSYNC_LOG_INDENT;
+
+		memset(TempBufferData, 0, TestFileSize);
+
+		FNativeFile		 TestFile(TestFilename, EFileMode::ReadOnlyUnbufferedOrdered);
+		constexpr uint64 ChunkSize = 1_MB;
+		static_assert(TestFileSize % ChunkSize == 0);
+
+		UNSYNC_LOG(L"Reading test data");
+
+		const FTimePoint ReadStartTime = TimePointNow();
+
+		const uint64 NumChunks = TestFileSize / ChunkSize;
+
+		FBlake3Hasher Hasher;
+
+		uint64 CurrentOffset = 0;
+
+		auto IoCallback = [TempBufferData, &Hasher, &CurrentOffset](FIOBuffer Buffer, uint64 SourceOffset, uint64 ReadSize, uint64 UserData)
+		{
+			UNSYNC_ASSERT(CurrentOffset == SourceOffset);
+			UNSYNC_ASSERT(ReadSize == Buffer.GetSize());
+			CurrentOffset += ReadSize;
+			Hasher.Update(Buffer.GetData(), Buffer.GetSize());
+		};
+
+		for (uint32 ChunkIndex = 0; ChunkIndex < NumChunks; ++ChunkIndex)
+		{
+			uint64 ChunkOffset = ChunkIndex * ChunkSize;
+			TestFile.ReadAsync(ChunkOffset, ChunkSize, 0, IoCallback);
+		}
+
+		TestFile.FlushAll();
+		const FTimePoint ReadDoneTime = TimePointNow();
+
+		FHash256 ActualHash = Hasher.Finalize();
+
+		const FTimePoint HashDoneTime = TimePointNow();
+
+		const double TotalDuration = DurationSec(ReadStartTime, HashDoneTime);
+
+		UNSYNC_LOG(L"Read + Hash rate: %.2f MB / sec", SizeMb(TestFileSize) / TotalDuration);
+		UNSYNC_LOG(L"Total time: %.3f sec", TotalDuration);
+
+		UNSYNC_ASSERT(ActualHash == ExpectedHash);
+	}
+}
+
 }  // namespace unsync
diff --git a/Engine/Source/Programs/Unsync/Private/UnsyncFile.h b/Engine/Source/Programs/Unsync/Private/UnsyncFile.h
--- a/Engine/Source/Programs/Unsync/Private/UnsyncFile.h
+++ b/Engine/Source/Programs/Unsync/Private/UnsyncFile.h
@@ -25,12 +25,14 @@ enum class EFileMode : uint32 {
 
 	// Extended modes
 	IgnoreDryRun = 1 << 4,	// allow write operations even in dry run mode
+	Ordered		 = 1 << 5,	// make async operations complete in order they are issued
 
 	// Commonly used mode combinations
-	ReadOnly		   = Read,
-	ReadOnlyUnbuffered = Read | Unbuffered,
-	CreateReadWrite	   = Read | Write | Create,
-	CreateWriteOnly	   = Write | Create,
+	ReadOnly				  = Read,
+	ReadOnlyUnbuffered		  = Read | Unbuffered,
+	ReadOnlyUnbufferedOrdered = ReadOnlyUnbuffered | Ordered,
+	CreateReadWrite			  = Read | Write | Create,
+	CreateWriteOnly			  = Write | Create,
 
 	// Masks
 	CommonModeMask	 = Create | Read | Write | Unbuffered,
@@ -189,9 +191,11 @@ struct FWindowsFile : FIOReaderWriter
 	HANDLE FileHandle = INVALID_HANDLE_VALUE;
 	int32  LastError  = 0;
 
-	static constexpr uint32 NUM_QUEUES					 = 16;
-	HANDLE					OverlappedEvents[NUM_QUEUES] = {};
-	struct Command
+	static constexpr uint32 MAX_OVERLAPPED_COMMANDS = 16;
+
+	HANDLE OverlappedEvents[MAX_OVERLAPPED_COMMANDS] = {};
+
+	struct FOverlappedCommand
 	{
 		OVERLAPPED Overlapped	= {};
 		uint64	   SourceOffset = 0;
@@ -201,7 +205,10 @@ struct FWindowsFile : FIOReaderWriter
 		FIOBuffer  Buffer;
 		IOCallback Callback = {};
 	};
-	Command Commands[NUM_QUEUES] = {};
+	FOverlappedCommand Commands[MAX_OVERLAPPED_COMMANDS] = {};
+
+	uint32 MaxQueueDepth = MAX_OVERLAPPED_COMMANDS;
+	uint64 NumCommandsIssued = 0;
 
 	FPath Filename;
 
@@ -211,7 +218,7 @@ private:
 
 	// All internal methods expect the Mutex to be locked
 	void   InternalFlushAll();
-	uint32 CompleteReadCommand(Command& Cmd);
+	uint32 CompleteReadCommand(FOverlappedCommand& Cmd);
 	bool   OpenFileHandle(EFileMode InMode);
 
 private:
@@ -350,7 +357,7 @@ struct FNullReaderWriter : FIOReaderWriter
 	explicit FNullReaderWriter(FInvalid) : DataSize(0), bValid(false) {}
 
 	// IOBase
-	virtual void   FlushAll() override {};
+	virtual void   FlushAll() override{};
 	virtual uint64 GetSize() override { return DataSize; }
 	virtual bool   IsValid() override { return bValid; }
 	virtual void   Close() override {};
diff --git a/Engine/Source/Programs/Unsync/Private/UnsyncTest.cpp b/Engine/Source/Programs/Unsync/Private/UnsyncTest.cpp
--- a/Engine/Source/Programs/Unsync/Private/UnsyncTest.cpp
+++ b/Engine/Source/Programs/Unsync/Private/UnsyncTest.cpp
@@ -858,6 +858,12 @@ RunTests(const std::string& Preset)
 		TestFiles();
 	}
 
+	if (Preset == "file_async_read" || Preset == "all")
+	{
+		extern void TestFileAsyncRead();
+		TestFileAsyncRead();
+	}
+
 	if (Preset == "build_target" || Preset == "all")
 	{
 		for (EWeakHashAlgorithmID Weak : WeakList)
diff --git a/Engine/Source/Programs/Unsync/Private/UnsyncVersion.cpp b/Engine/Source/Programs/Unsync/Private/UnsyncVersion.cpp
--- a/Engine/Source/Programs/Unsync/Private/UnsyncVersion.cpp
+++ b/Engine/Source/Programs/Unsync/Private/UnsyncVersion.cpp
@@ -2,7 +2,7 @@
 
 #include "UnsyncVersion.h"
 
-#define UNSYNC_VERSION_STR "1.0.92"
+#define UNSYNC_VERSION_STR "1.0.93"
 
 namespace unsync {
 
