unsync - Add a dedicated FAsyncReader object instead of FIOReader::ReadAsync()

* This allows multiple independent ordered async read operations to be in flight simultaneously
* Always use strictly ordered async reads
* Remove automatic recovery on async read errors

diff --git a/Engine/Source/Programs/Unsync/Private/UnsyncChunking.cpp b/Engine/Source/Programs/Unsync/Private/UnsyncChunking.cpp
--- a/Engine/Source/Programs/Unsync/Private/UnsyncChunking.cpp
+++ b/Engine/Source/Programs/Unsync/Private/UnsyncChunking.cpp
@@ -385,6 +385,8 @@ ComputeBlocksFixedT(FIOReader& Reader, const FComputeBlocksParams& Params)
 	{
 		FSchedulerSemaphore IoSemaphore(*GScheduler, 16);
 		FTaskGroup			TaskGroup = GScheduler->CreateTaskGroup(&IoSemaphore);
+		
+		std::unique_ptr<FAsyncReader> AsyncReader = Reader.CreateAsyncReader();
 
 		for (uint64 I = 0; I < NumReads; ++I)
 		{
@@ -438,9 +440,9 @@ ComputeBlocksFixedT(FIOReader& Reader, const FComputeBlocksParams& Params)
 					});
 			};
 
-			Reader.ReadAsync(Offset, ThisReadSize, 0, ReadCallback);
+			AsyncReader->EnqueueRead(Offset, ThisReadSize, 0, ReadCallback);
 		}
-		Reader.FlushAll();
+		AsyncReader->Flush();
 		TaskGroup.wait();
 	}
 
diff --git a/Engine/Source/Programs/Unsync/Private/UnsyncCmdPack.cpp b/Engine/Source/Programs/Unsync/Private/UnsyncCmdPack.cpp
--- a/Engine/Source/Programs/Unsync/Private/UnsyncCmdPack.cpp
+++ b/Engine/Source/Programs/Unsync/Private/UnsyncCmdPack.cpp
@@ -691,36 +691,37 @@ BuildTargetFromPack(FIOWriter& Output, const FPackDatabase& PackDb, TArrayView<F
 		Output.Write(Decompressed.GetData(), Item.Block->Offset, Item.Block->Size);
 	};
 
+	std::unique_ptr<FAsyncReader> AsyncReader;
+
 	for (uint64 ScheduleIndex = 0; ScheduleIndex < Schedule.size(); ++ScheduleIndex)
 	{
 		const FScheduleItem& Item = Schedule[ScheduleIndex];
 		if (CurrentPackId != Item.PackIndex)
 		{
-			if (PackFile)
-			{
-				PackFile->FlushAll();
-			}
+			std::shared_ptr<FNativeFile> NextPackFile = PackDb.GetPackFile(Item.PackIndex);
 
-			PackFile = PackDb.GetPackFile(Item.PackIndex);
-
-			if (!PackFile->IsValid())
+			if (!NextPackFile->IsValid())
 			{
 				UNSYNC_ERROR(L"Failed to open pack file '%ls'", PackDb.PackFilenames[Item.PackIndex].wstring().c_str());
 				return false;
 			}
 
+			AsyncReader	  = NextPackFile->CreateAsyncReader();
+			std::swap(PackFile, NextPackFile);
+
 			CurrentPackId = Item.PackIndex;
 		}
 
 		UNSYNC_ASSERT(Item.IndexEntry->PackBlockOffset < PackFile->GetSize());
 		UNSYNC_ASSERT(Item.IndexEntry->PackBlockOffset + Item.IndexEntry->PackBlockSize <= PackFile->GetSize());
-		PackFile->ReadAsync(Item.IndexEntry->PackBlockOffset, Item.IndexEntry->PackBlockSize, ScheduleIndex, ReadCallback);
+		AsyncReader->EnqueueRead(Item.IndexEntry->PackBlockOffset, Item.IndexEntry->PackBlockSize, ScheduleIndex, ReadCallback);
 	}
 
-	if (PackFile)
+	if (AsyncReader)
 	{
-		PackFile->FlushAll();
+		AsyncReader->Flush();
 	}
+	AsyncReader = {};
 
 	return true;
 }
diff --git a/Engine/Source/Programs/Unsync/Private/UnsyncCore.cpp b/Engine/Source/Programs/Unsync/Private/UnsyncCore.cpp
--- a/Engine/Source/Programs/Unsync/Private/UnsyncCore.cpp
+++ b/Engine/Source/Programs/Unsync/Private/UnsyncCore.cpp
@@ -204,6 +204,8 @@ ValidateTarget(FIOReader& Reader, const FNeedList& NeedList, EStrongHashAlgorith
 	uint64 BatchBegin	  = 0;
 	uint64 BatchSizeBytes = 0;
 
+	std::unique_ptr<FAsyncReader> AsyncReader = Reader.CreateAsyncReader();
+
 	for (uint64 BlockIndex = 0; BlockIndex < ValidationBlocks.size(); ++BlockIndex)
 	{
 		const FGenericBlock& CurrBlock = ValidationBlocks[BlockIndex];
@@ -278,7 +280,7 @@ ValidateTarget(FIOReader& Reader, const FNeedList& NeedList, EStrongHashAlgorith
 			});
 		};
 
-		Reader.ReadAsync(ReadOffset, BatchSizeBytes, 0, ReadCallback);
+		AsyncReader->EnqueueRead(ReadOffset, BatchSizeBytes, 0, ReadCallback);
 
 		if (NumInvalidBlocks)
 		{
@@ -289,7 +291,7 @@ ValidateTarget(FIOReader& Reader, const FNeedList& NeedList, EStrongHashAlgorith
 		BatchBegin	   = BlockIndex + 1;
 	}
 
-	Reader.FlushAll();
+	AsyncReader->Flush();
 	TaskGroup.wait();
 
 	ValidationProgressLogger.Complete();
diff --git a/Engine/Source/Programs/Unsync/Private/UnsyncDiff.cpp b/Engine/Source/Programs/Unsync/Private/UnsyncDiff.cpp
--- a/Engine/Source/Programs/Unsync/Private/UnsyncDiff.cpp
+++ b/Engine/Source/Programs/Unsync/Private/UnsyncDiff.cpp
@@ -159,6 +159,8 @@ DiffBlocksParallelT(FIOReader&				  BaseDataReader,
 	FSchedulerSemaphore IoSemaphore(*GScheduler, 16);
 	FTaskGroup			TaskGroup = GScheduler->CreateTaskGroup(&IoSemaphore);
 
+	std::unique_ptr<FAsyncReader> AsyncReader = BaseDataReader.CreateAsyncReader();
+
 	for (uint64 I = 0; I < NumTasks; ++I)
 	{
 		FTask& Task		 = Tasks[I];
@@ -248,10 +250,10 @@ DiffBlocksParallelT(FIOReader&				  BaseDataReader,
 				});
 		};
 
-		BaseDataReader.ReadAsync(Task.Offset, Task.Size, I, ReadCallback);
+		AsyncReader->EnqueueRead(Task.Offset, Task.Size, I, ReadCallback);
 	}
 
-	BaseDataReader.FlushAll();
+	AsyncReader->Flush();
 	TaskGroup.wait();
 
 	THashSet<FGenericBlock, FBlockStrongHash, FBlockStrongHashEq> BaseStrongHashSet;
diff --git a/Engine/Source/Programs/Unsync/Private/UnsyncFile.cpp b/Engine/Source/Programs/Unsync/Private/UnsyncFile.cpp
--- a/Engine/Source/Programs/Unsync/Private/UnsyncFile.cpp
+++ b/Engine/Source/Programs/Unsync/Private/UnsyncFile.cpp
@@ -313,6 +313,75 @@ struct FCreateFileInfo
 	}
 };
 
+struct FWindowsAsyncFileReader final : FAsyncReader
+{
+	UNSYNC_DISALLOW_COPY_ASSIGN(FWindowsAsyncFileReader)
+
+	struct FOverlappedCommand
+	{
+		OVERLAPPED Overlapped	= {};
+		uint64	   SourceOffset = 0;
+		uint64	   ReadSize		= 0;
+		uint64	   UserData		= 0;
+		uint32	   ErrorCode	= 0;
+		bool	   bActive		= false;
+		FIOBuffer  Buffer;
+		IOCallback Callback = {};
+	};
+
+	FWindowsAsyncFileReader(FWindowsFile& InReader, uint32 InMaxPipelineDepth);
+
+	virtual ~FWindowsAsyncFileReader();
+	virtual uint64 GetSize() override { return Inner.GetSize(); }
+	virtual bool   IsValid() override { return Inner.IsValid(); }
+	virtual bool   EnqueueRead(uint64 SourceOffset, uint64 Size, uint64 UserData, IOCallback Callback);
+	virtual void   Flush() override;
+
+	// NOTE: Expects that Inner.Mutex is locked
+	uint32 CompleteReadCommand(FOverlappedCommand& Cmd);
+
+	FWindowsFile& Inner;
+	const uint32  MaxQueueDepth;
+
+	static constexpr uint32 MAX_OVERLAPPED_COMMANDS = MAX_IO_PIPELINE_DEPTH;
+
+	HANDLE			   OverlappedEvents[MAX_OVERLAPPED_COMMANDS] = {};
+	FOverlappedCommand Commands[MAX_OVERLAPPED_COMMANDS]		 = {};
+	uint64			   NumCommandsIssued						 = 0;
+
+	std::mutex Mutex;
+};
+
+FWindowsAsyncFileReader::FWindowsAsyncFileReader(FWindowsFile& InReader, uint32 InMaxPipelineDepth)
+: Inner(InReader)
+, MaxQueueDepth(InMaxPipelineDepth)
+{
+	for (uint32 I = 0; I < MaxQueueDepth; ++I)
+	{
+		OverlappedEvents[I]			  = CreateEvent(nullptr, true, true, nullptr);
+		Commands[I].Overlapped.hEvent = OverlappedEvents[I];
+	}
+
+	std::lock_guard<std::mutex> LockGuard(Inner.Mutex);
+	Inner.AsyncReaders.push_back(this);
+}
+
+FWindowsAsyncFileReader::~FWindowsAsyncFileReader()
+{
+	Flush();
+
+	for (HANDLE EventHandle : OverlappedEvents)
+	{
+		if (EventHandle)
+		{
+			CloseHandle(EventHandle);
+		}
+	}
+
+	std::lock_guard<std::mutex> LockGuard(Inner.Mutex);
+	Inner.AsyncReaders.erase(std::remove(Inner.AsyncReaders.begin(), Inner.AsyncReaders.end(), this), Inner.AsyncReaders.end());
+}
+
 FWindowsFile::FWindowsFile(const FPath& InFilename, EFileMode InMode, uint64 InSize) : Mode(InMode)
 {
 	Filename = MakeExtendedAbsolutePath(InFilename);
@@ -370,17 +439,14 @@ FWindowsFile::FWindowsFile(const FPath& InFilename, EFileMode InMode, uint64 InS
 		{
 			UNSYNC_ERROR(L"Unexpected file mode %d", (int)InMode);
 		}
-
-		for (uint32 I = 0; I < MaxQueueDepth; ++I)
-		{
-			OverlappedEvents[I]			  = CreateEvent(nullptr, true, true, nullptr);
-			Commands[I].Overlapped.hEvent = OverlappedEvents[I];
-		}
 	}
 }
 FWindowsFile::~FWindowsFile()
 {
 	Close();
+	
+	std::lock_guard<std::mutex> LockGuard(Mutex);
+	UNSYNC_ASSERT(AsyncReaders.empty());
 }
 
 bool
@@ -416,90 +482,43 @@ FWindowsFile::IsValid()
 void
 FWindowsFile::Close()
 {
-	InternalFlushAll();
+	FlushAsyncReaders();
 
+	std::lock_guard<std::mutex> LockGuard(Mutex);
 	if (FileHandle != INVALID_HANDLE_VALUE)
 	{
 		CloseHandle(FileHandle);
 		FileHandle = INVALID_HANDLE_VALUE;
 	}
-	for (uint32 I = 0; I < MaxQueueDepth; ++I)
-	{
-		if (OverlappedEvents[I])
-		{
-			CloseHandle(OverlappedEvents[I]);
-			OverlappedEvents[I] = 0;
-		}
-	}
 }
 
 uint32
-FWindowsFile::CompleteReadCommand(FOverlappedCommand& Cmd)
+FWindowsAsyncFileReader::CompleteReadCommand(FOverlappedCommand& Cmd)
 {
-	// Expects that Mutex is locked
+	// Expects that Inner.Mutex is locked
 
 	UNSYNC_ASSERT(Cmd.bActive);
-	const uint32 MaxAttempts = 100;
 
 	DWORD ReadBytes = 0;
 
-	bool bRecoveredFromError = false;
-
-	const uint64 ExpectedReadEnd   = std::min(FileSize, Cmd.SourceOffset + Cmd.ReadSize);
+	const uint64 ExpectedReadEnd   = std::min(Inner.FileSize, Cmd.SourceOffset + Cmd.ReadSize);
 	const uint64 ExpectedReadBytes = ExpectedReadEnd - Cmd.SourceOffset;
 
-	for (uint32 Attempt = 0; Attempt < MaxAttempts; ++Attempt)
-	{
-		BOOL OverlappedResultOk = GetOverlappedResult(FileHandle, &Cmd.Overlapped, &ReadBytes, true);
-
-		LastError = 0;
-		if (!OverlappedResultOk)
-		{
-			LastError = GetLastError();
-		}
-
-		if (ReadBytes == ExpectedReadBytes || ReadBytes == Cmd.ReadSize)
-		{
-			break;
-		}
-		else
-		{
-			UNSYNC_WARNING(L"FNativeFile expected to read %lld bytes, but read %lld. %hs",
-						   (uint64)ExpectedReadBytes,
-						   (uint64)ReadBytes,
-						   FormatSystemErrorMessage(LastError).c_str());
-
-			UNSYNC_LOG(L"Trying to recover from error (attempt %d of %d)", Attempt + 1, MaxAttempts);
-
-			SchedulerSleep(1000);
-
-			// Cancel any active overlapped IO commands
-			for (FOverlappedCommand& It : Commands)
-			{
-				if (It.bActive && (&It != &Cmd))
-				{
-					CancelIoEx(FileHandle, &It.Overlapped);
-					WaitForSingleObject(It.Overlapped.hEvent, INFINITE);
-				}
-			}
+	BOOL OverlappedResultOk = GetOverlappedResult(Inner.FileHandle, &Cmd.Overlapped, &ReadBytes, true);
 
-			// Re-open the file handle
-			CloseHandle(FileHandle);
-			bool bOpenedOk = OpenFileHandle(Mode);
-			if (!bOpenedOk)
-			{
-				LastError = GetLastError();
-				UNSYNC_ERROR(L"Failed to re-open the file. %hs", FormatSystemErrorMessage(LastError).c_str());
-				break;
-			}
+	if (OverlappedResultOk)
+	{
+		Cmd.ErrorCode = 0;
+	}
+	else
+	{
+		Cmd.ErrorCode = GetLastError();
+	}
 
-			// Try reading this block again
-			ResetEvent(Cmd.Overlapped.hEvent);
-			ReadFile(FileHandle, Cmd.Buffer.GetMemory(), CheckedNarrow(Cmd.ReadSize), nullptr, &Cmd.Overlapped);
+	Inner.LastError = Cmd.ErrorCode;
 
-			bRecoveredFromError = true;
-		}
-	}
+	// TODO: retry read command a few times to retry from errors
+	UNSYNC_ASSERT(ReadBytes == ExpectedReadBytes);
 
 	const uint64 ReadBytesClamped = std::min<uint64>(Cmd.Buffer.GetSize(), ReadBytes);
 
@@ -510,19 +529,6 @@ FWindowsFile::CompleteReadCommand(FOverlappedCommand& Cmd)
 
 	Cmd.bActive = false;
 
-	if (bRecoveredFromError)
-	{
-		// Re-issue the read requests using new file handle
-		for (FOverlappedCommand& It : Commands)
-		{
-			if (It.bActive)
-			{
-				ResetEvent(It.Overlapped.hEvent);
-				ReadFile(FileHandle, It.Buffer.GetMemory(), CheckedNarrow(It.ReadSize), nullptr, &It.Overlapped);
-			}
-		}
-	}
-
 	return CheckedNarrow(ReadBytesClamped);
 }
 
@@ -537,7 +543,7 @@ FWindowsFile::Write(const void* Data, uint64 DestOffset, uint64 TotalSize)
 
 	if (!IsWriteOnly(Mode))
 	{
-		InternalFlushAll();	 // flush any outstanding read requests before writing
+		FlushAsyncReaders();	 // flush any outstanding read requests before writing
 	}
 
 	LARGE_INTEGER Pos;
@@ -605,7 +611,10 @@ FWindowsFile::Read(void* Dest, uint64 SourceOffset, uint64 ReadSize)
 {
 	std::lock_guard<std::mutex> LockGuard(Mutex);
 
-	UNSYNC_ASSERTF((Mode & EFileMode::Unbuffered) == 0, L"Unbuffered files only support ReadAsync");
+	UNSYNC_ASSERTF(((Mode & EFileMode::Unbuffered) == 0) ||
+					   ((SourceOffset % UNBUFFERED_READ_ALIGNMENT == 0) && (ReadSize % UNBUFFERED_READ_ALIGNMENT == 0)),
+				   L"Unbuffered files only support Read when offset and size are aligned to 4KB");
+
 	UNSYNC_ASSERT(IsReadable(Mode));
 
 	LARGE_INTEGER Pos;
@@ -650,37 +659,26 @@ FWindowsFile::Read(void* Dest, uint64 SourceOffset, uint64 ReadSize)
 }
 
 bool
-FWindowsFile::ReadAsync(uint64 SourceOffset, uint64 Size, uint64 UserData, IOCallback Callback)
+FWindowsAsyncFileReader::EnqueueRead(uint64 SourceOffset, uint64 Size, uint64 UserData, IOCallback Callback)
 {
-	std::lock_guard<std::mutex> LockGuard(Mutex);
+	std::lock_guard<std::mutex> LockGuard(Inner.Mutex);
+
+	const EFileMode Mode = Inner.Mode;
 
 	UNSYNC_ASSERT(IsReadable(Mode));
 
 	uint32 CmdIdx = ~0u;
 
-	const bool bIsOrdered = !!(Mode & EFileMode::Ordered);
+	// Async commands are always strictly ordered
+	const uint32 WaitSlotIndex = uint32(NumCommandsIssued % MaxQueueDepth);
+	const DWORD WaitResult = WaitForSingleObject(OverlappedEvents[WaitSlotIndex], INFINITE);
+	UNSYNC_ASSERT(WaitResult == WAIT_OBJECT_0);
+	CmdIdx = WaitSlotIndex;
 
-	if (bIsOrdered)
-	{
-		const uint32 WaitSlotIndex = uint32(NumCommandsIssued % MaxQueueDepth);
-		const DWORD WaitResult = WaitForSingleObject(OverlappedEvents[WaitSlotIndex], INFINITE);
-		UNSYNC_ASSERT(WaitResult == WAIT_OBJECT_0);
-		CmdIdx = WaitSlotIndex;
-	}
-	else
-	{
-		// find available queue or wait for one to complete
-		const DWORD WaitResult = WaitForMultipleObjects(MaxQueueDepth, OverlappedEvents, false, INFINITE);
-
-		if (WaitResult >= WAIT_OBJECT_0 && WaitResult < MaxQueueDepth)
-		{
-			CmdIdx = WaitResult - WAIT_OBJECT_0;
-		}
-	}
 
 	UNSYNC_ASSERT(CmdIdx < MaxQueueDepth);
 
-	FWindowsFile::FOverlappedCommand& Cmd = Commands[CmdIdx];
+	FOverlappedCommand& Cmd = Commands[CmdIdx];
 
 	if (Cmd.bActive)
 	{
@@ -693,8 +691,8 @@ FWindowsFile::ReadAsync(uint64 SourceOffset, uint64 Size, uint64 UserData, IOCal
 		uint64 OriginalBegin = SourceOffset;
 		uint64 OriginalEnd	 = SourceOffset + Size;
 
-		uint64 AlignedBegin = AlignDownToMultiplePow2(OriginalBegin, UNBUFFERED_READ_ALIGNMENT);
-		uint64 AlignedEnd	= AlignUpToMultiplePow2(OriginalEnd, UNBUFFERED_READ_ALIGNMENT);
+		uint64 AlignedBegin = AlignDownToMultiplePow2(OriginalBegin, FWindowsFile::UNBUFFERED_READ_ALIGNMENT);
+		uint64 AlignedEnd	= AlignUpToMultiplePow2(OriginalEnd, FWindowsFile::UNBUFFERED_READ_ALIGNMENT);
 
 		SourceOffset = AlignedBegin;
 		Size		 = AlignedEnd - AlignedBegin;
@@ -719,10 +717,10 @@ FWindowsFile::ReadAsync(uint64 SourceOffset, uint64 Size, uint64 UserData, IOCal
 	Cmd.Callback			  = Callback;
 
 	ResetEvent(Cmd.Overlapped.hEvent);
-	DWORD ReadOk = ReadFile(FileHandle, Cmd.Buffer.GetMemory(), CheckedNarrow(Size), nullptr, &Cmd.Overlapped);
+	DWORD ReadOk = ReadFile(Inner.FileHandle, Cmd.Buffer.GetMemory(), CheckedNarrow(Size), nullptr, &Cmd.Overlapped);
 	if (!ReadOk)
 	{
-		LastError = GetLastError();
+		Inner.LastError = GetLastError();
 	}
 	Cmd.bActive = true;
 
@@ -732,8 +730,10 @@ FWindowsFile::ReadAsync(uint64 SourceOffset, uint64 Size, uint64 UserData, IOCal
 }
 
 void
-FWindowsFile::InternalFlushAll()
+FWindowsAsyncFileReader::Flush()
 {
+	std::lock_guard<std::mutex> LockGuard(Inner.Mutex);
+
 	for (uint32 I = 0; I < MaxQueueDepth; ++I)
 	{
 		const uint64 CommandIndex = (NumCommandsIssued + I) % MaxQueueDepth;
@@ -744,12 +744,19 @@ FWindowsFile::InternalFlushAll()
 	}
 }
 
-void
-FWindowsFile::FlushAll()
+void FWindowsFile::FlushAsyncReaders()
 {
-	std::lock_guard<std::mutex> LockGuard(Mutex);
+	for (FWindowsAsyncFileReader* It : AsyncReaders)
+	{
+		It->Flush();
+	}
+}
 
-	InternalFlushAll();
+std::unique_ptr<FAsyncReader>
+FWindowsFile::CreateAsyncReader(uint32 MaxPipelineDepth)
+{
+	MaxPipelineDepth = std::min(MaxPipelineDepth, FWindowsAsyncFileReader::MAX_OVERLAPPED_COMMANDS);
+	return std::unique_ptr<FAsyncReader>(new FWindowsAsyncFileReader(*this, MaxPipelineDepth));
 }
 
 FFileAttributes
@@ -841,12 +848,6 @@ FUnixFile::~FUnixFile()
 	Close();
 }
 
-void
-FUnixFile::FlushAll()
-{
-	// nothing to do, all file IO is synchronous
-}
-
 void
 FUnixFile::Close()
 {
@@ -869,22 +870,11 @@ FUnixFile::Read(void* dest, uint64 SourceOffset, uint64 ReadSize)
 	return read_bytes;
 }
 
-bool
-FUnixFile::ReadAsync(uint64 SourceOffset, uint64 size, uint64 UserData, IOCallback callback)
-{
-	return FIOReader::ReadAsync(SourceOffset, size, UserData, callback);
-}
-
 uint64
 FUnixFile::Write(const void* data, uint64 DestOffset, uint64 WriteSize)
 {
 	UNSYNC_ASSERT(IsWritable(Mode));
 
-	if (!IsWriteOnly(Mode))
-	{
-		FlushAll();	 // flush any outstanding read requests before writing
-	}
-
 	// TODO: handle partial writes (pwrite returning 0 < x < WriteSize)
 	uint64 wrote_bytes = pwrite(FileDescriptor, data, WriteSize, DestOffset);
 	if (wrote_bytes != WriteSize)
@@ -1524,12 +1514,14 @@ TestFileTime()
 }
 
 uint64
-BlockingReadLarge(FIOReader& Reader, uint64 Offset, uint64 Size, uint8* OutputBuffer, uint64 OutputBufferSize)
+BlockingReadLarge(FIOReader& InReader, uint64 Offset, uint64 Size, uint8* OutputBuffer, uint64 OutputBufferSize)
 {
 	const uint64 BytesPerRead = 2_MB;
-	const uint64 ReadEnd	  = std::min(Offset + Size, Reader.GetSize());
+	const uint64 ReadEnd	  = std::min(Offset + Size, InReader.GetSize());
 	const uint64 ClampedSize  = ReadEnd - Offset;
 
+	std::unique_ptr<FAsyncReader> AsyncReader = InReader.CreateAsyncReader();
+
 	std::atomic<uint64> TotalReadSize = 0;
 
 	if (ClampedSize == 0)
@@ -1562,10 +1554,10 @@ BlockingReadLarge(FIOReader& Reader, uint64 Offset, uint64 Size, uint8* OutputBu
 				});
 		};
 
-		Reader.ReadAsync(ThisReadOffset, ThisBatchSize, OutputOffset, ReadCallback);
+		AsyncReader->EnqueueRead(ThisReadOffset, ThisBatchSize, OutputOffset, ReadCallback);
 	}
 
-	Reader.FlushAll();
+	AsyncReader->Flush();
 	CopyTasks.wait();
 
 	return TotalReadSize;
@@ -1798,6 +1790,32 @@ SystemTimeFromFileTime(std::filesystem::file_time_type FileTime)
 #endif // 0
 }
 
+FDummyAsyncReader::FDummyAsyncReader(FIOReader& InReader) : Inner(InReader)
+{
+}
+
+bool
+FDummyAsyncReader::EnqueueRead(uint64 SourceOffset, uint64 Size, uint64 UserData, IOCallback Callback)
+{
+	if (Size != 0)
+	{
+		FIOBuffer Buffer   = FIOBuffer::Alloc(Size, L"FDummyAsyncReader::ReadAsync");
+		uint64	  ReadSize = Inner.Read(Buffer.GetData(), SourceOffset, Size);
+		Callback(std::move(Buffer), SourceOffset, ReadSize, UserData);
+		return true;
+	}
+	else
+	{
+		return false;
+	}
+}
+
+std::unique_ptr<FAsyncReader>
+FIOReader::CreateAsyncReader(uint32 MaxPipelineDepth)
+{
+	return std::unique_ptr<FDummyAsyncReader>(new FDummyAsyncReader(*this));
+}
+
 void
 TestFileAsyncRead()
 {
@@ -1842,7 +1860,9 @@ TestFileAsyncRead()
 
 		memset(TempBufferData, 0, TestFileSize);
 
-		FNativeFile		 TestFile(TestFilename, EFileMode::ReadOnlyUnbuffered);
+		FNativeFile					  TestFile(TestFilename, EFileMode::ReadOnlyUnbuffered);
+		std::unique_ptr<FAsyncReader> TestFileReader = TestFile.CreateAsyncReader();
+
 		constexpr uint64 ChunkSize = 8_MB;
 		static_assert(TestFileSize % ChunkSize == 0);
 
@@ -1859,10 +1879,10 @@ TestFileAsyncRead()
 		{
 			uint64 ChunkOffset = ChunkIndex * ChunkSize;
 
-			TestFile.ReadAsync(ChunkOffset, ChunkSize, 0, IoCallback);
+			TestFileReader->EnqueueRead(ChunkOffset, ChunkSize, 0, IoCallback);
 		}
 
-		TestFile.FlushAll();
+		TestFileReader->Flush();
 
 		const FTimePoint ReadDoneTime = TimePointNow();
 
@@ -1884,14 +1904,16 @@ TestFileAsyncRead()
 		UNSYNC_ASSERT(ActualHash == ExpectedHash);
 	}
 
-	UNSYNC_LOG(L"ReadOnlyUnbufferedOrdered");
+	UNSYNC_LOG(L"ReadOnlyUnbufferedStreaming");
 
 	{
 		UNSYNC_LOG_INDENT;
 
 		memset(TempBufferData, 0, TestFileSize);
 
-		FNativeFile		 TestFile(TestFilename, EFileMode::ReadOnlyUnbufferedOrdered);
+		FNativeFile					  TestFile(TestFilename, EFileMode::ReadOnlyUnbuffered);
+		std::unique_ptr<FAsyncReader> TestFileReader = TestFile.CreateAsyncReader();
+
 		constexpr uint64 ChunkSize = 1_MB;
 		static_assert(TestFileSize % ChunkSize == 0);
 
@@ -1916,10 +1938,10 @@ TestFileAsyncRead()
 		for (uint32 ChunkIndex = 0; ChunkIndex < NumChunks; ++ChunkIndex)
 		{
 			uint64 ChunkOffset = ChunkIndex * ChunkSize;
-			TestFile.ReadAsync(ChunkOffset, ChunkSize, 0, IoCallback);
+			TestFileReader->EnqueueRead(ChunkOffset, ChunkSize, 0, IoCallback);
 		}
 
-		TestFile.FlushAll();
+		TestFileReader->Flush();
 		const FTimePoint ReadDoneTime = TimePointNow();
 
 		FHash256 ActualHash = Hasher.Finalize();
diff --git a/Engine/Source/Programs/Unsync/Private/UnsyncFile.h b/Engine/Source/Programs/Unsync/Private/UnsyncFile.h
--- a/Engine/Source/Programs/Unsync/Private/UnsyncFile.h
+++ b/Engine/Source/Programs/Unsync/Private/UnsyncFile.h
@@ -15,6 +15,8 @@ namespace unsync {
 
 extern bool GForceBufferedFiles;
 
+static constexpr uint32 MAX_IO_PIPELINE_DEPTH = 16;
+
 enum class EFileMode : uint32 {
 	None	   = 0,
 
@@ -25,12 +27,10 @@ enum class EFileMode : uint32 {
 
 	// Extended modes
 	IgnoreDryRun = 1 << 4,	// allow write operations even in dry run mode
-	Ordered		 = 1 << 5,	// make async operations complete in order they are issued
 
 	// Commonly used mode combinations
 	ReadOnly				  = Read,
 	ReadOnlyUnbuffered		  = Read | Unbuffered,
-	ReadOnlyUnbufferedOrdered = ReadOnlyUnbuffered | Ordered,
 	CreateReadWrite			  = Read | Write | Create,
 	CreateWriteOnly			  = Write | Create,
 
@@ -132,30 +132,39 @@ using IOCallback = std::function<void(FIOBuffer Buffer, uint64 SourceOffset, uin
 struct FIOBase
 {
 	virtual ~FIOBase()		  = default;
-	virtual void   FlushAll() = 0;
 	virtual uint64 GetSize()  = 0;
 	virtual bool   IsValid()  = 0;
 	virtual void   Close()	  = 0;
 	virtual int32  GetError() = 0;
 };
 
+struct FAsyncReader
+{
+	virtual ~FAsyncReader()																			   = default;
+	virtual uint64 GetSize()																		   = 0;
+	virtual bool   IsValid()																		   = 0;
+	virtual bool   EnqueueRead(uint64 SourceOffset, uint64 Size, uint64 UserData, IOCallback Callback) = 0;
+	virtual void   Flush()																			   = 0;
+};
+
 struct FIOReader : virtual FIOBase
 {
 	virtual uint64 Read(void* Dest, uint64 SourceOffset, uint64 Size) = 0;
-	virtual bool   ReadAsync(uint64 SourceOffset, uint64 Size, uint64 UserData, IOCallback Callback)
-	{
-		if (Size != 0)
-		{
-			FIOBuffer Buffer   = FIOBuffer::Alloc(Size, L"FIOReader::ReadAsync");
-			uint64	  ReadSize = Read(Buffer.GetData(), SourceOffset, Size);
-			Callback(std::move(Buffer), SourceOffset, ReadSize, UserData);
-			return true;
-		}
-		else
-		{
-			return false;
-		}
-	}
+
+	virtual std::unique_ptr<FAsyncReader> CreateAsyncReader(uint32 MaxPipelineDepth = MAX_IO_PIPELINE_DEPTH);
+};
+
+struct FDummyAsyncReader final : FAsyncReader
+{
+	FDummyAsyncReader(FIOReader& InReader);
+
+	virtual uint64 GetSize() override { return Inner.GetSize(); }
+	virtual bool   IsValid() override { return Inner.IsValid(); }
+	virtual bool   EnqueueRead(uint64 SourceOffset, uint64 Size, uint64 UserData, IOCallback Callback) override;
+	virtual void   Flush() override {};
+
+private:
+	FIOReader& Inner;
 };
 
 struct FIOWriter : virtual FIOBase
@@ -170,19 +179,20 @@ struct FIOReaderWriter : FIOReader, FIOWriter
 #if UNSYNC_PLATFORM_WINDOWS
 struct FWindowsFile : FIOReaderWriter
 {
+	friend struct FWindowsAsyncFileReader;
+
 	FWindowsFile(const FPath& Filename, EFileMode Mode = EFileMode::ReadOnly, uint64 InSize = 0);
 	~FWindowsFile();
 
 	// IOBase
-	virtual void   FlushAll() override;
 	virtual uint64 GetSize() override { return FileSize; }
 	virtual bool   IsValid() override;
 	virtual void   Close() override;
 	virtual int32  GetError() override { return LastError; }
 
 	// IORead
-	virtual uint64 Read(void* Dest, uint64 SourceOffset, uint64 ReadSize) override;
-	virtual bool   ReadAsync(uint64 SourceOffset, uint64 Size, uint64 UserData, IOCallback Callback) override;
+	virtual uint64						  Read(void* Dest, uint64 SourceOffset, uint64 ReadSize) override;
+	virtual std::unique_ptr<FAsyncReader> CreateAsyncReader(uint32 MaxPipelineDepth = MAX_IO_PIPELINE_DEPTH) override;
 
 	// IOWrite
 	virtual uint64 Write(const void* InData, uint64 DestOffset, uint64 WriteSize) override;
@@ -191,25 +201,6 @@ struct FWindowsFile : FIOReaderWriter
 	HANDLE FileHandle = INVALID_HANDLE_VALUE;
 	int32  LastError  = 0;
 
-	static constexpr uint32 MAX_OVERLAPPED_COMMANDS = 16;
-
-	HANDLE OverlappedEvents[MAX_OVERLAPPED_COMMANDS] = {};
-
-	struct FOverlappedCommand
-	{
-		OVERLAPPED Overlapped	= {};
-		uint64	   SourceOffset = 0;
-		uint64	   ReadSize		= 0;
-		uint64	   UserData		= 0;
-		bool	   bActive		= false;
-		FIOBuffer  Buffer;
-		IOCallback Callback = {};
-	};
-	FOverlappedCommand Commands[MAX_OVERLAPPED_COMMANDS] = {};
-
-	uint32 MaxQueueDepth = MAX_OVERLAPPED_COMMANDS;
-	uint64 NumCommandsIssued = 0;
-
 	FPath Filename;
 
 	static constexpr uint32 UNBUFFERED_READ_ALIGNMENT = 4096;
@@ -217,14 +208,15 @@ struct FWindowsFile : FIOReaderWriter
 private:
 
 	// All internal methods expect the Mutex to be locked
-	void   InternalFlushAll();
-	uint32 CompleteReadCommand(FOverlappedCommand& Cmd);
 	bool   OpenFileHandle(EFileMode InMode);
 
 private:
 	EFileMode Mode;
 
 	std::mutex Mutex;
+
+	void FlushAsyncReaders();
+	std::vector<FWindowsAsyncFileReader*> AsyncReaders;
 };
 using FNativeFile = FWindowsFile;
 #endif	// UNSYNC_PLATFORM_WINDOWS
@@ -236,7 +228,6 @@ struct FUnixFile : FIOReaderWriter
 	~FUnixFile();
 
 	// IOBase
-	virtual void   FlushAll() override;
 	virtual uint64 GetSize() override { return FileSize; }
 	virtual bool   IsValid() override { return FileHandle != nullptr; }
 	virtual void   Close() override;
@@ -244,7 +235,6 @@ struct FUnixFile : FIOReaderWriter
 
 	// IORead
 	virtual uint64 Read(void* Dest, uint64 SourceOffset, uint64 ReadSize) override;
-	virtual bool   ReadAsync(uint64 SourceOffset, uint64 ReadSize, uint64 UserData, IOCallback Callback) override;
 
 	// IOWrite
 	virtual uint64 Write(const void* Indata, uint64 DestOffset, uint64 WriteSize) override;
@@ -299,7 +289,6 @@ struct FMemReader : FIOReader
 	FMemReader(const uint8* InData, uint64 InDataSize);
 
 	// IOBase
-	virtual void   FlushAll() override {}
 	virtual uint64 GetSize() override { return Size; }
 	virtual bool   IsValid() override { return Data != nullptr; }
 	virtual void   Close() override
@@ -357,7 +346,6 @@ struct FNullReaderWriter : FIOReaderWriter
 	explicit FNullReaderWriter(FInvalid) : DataSize(0), bValid(false) {}
 
 	// IOBase
-	virtual void   FlushAll() override{};
 	virtual uint64 GetSize() override { return DataSize; }
 	virtual bool   IsValid() override { return bValid; }
 	virtual void   Close() override {};
@@ -369,13 +357,6 @@ struct FNullReaderWriter : FIOReaderWriter
 		memset(Dest, 0, ReadSize);
 		return ReadSize;
 	}
-	virtual bool ReadAsync(uint64 SourceOffset, uint64 Size, uint64 UserData, IOCallback Callback) override
-	{
-		FIOBuffer Buffer   = FIOBuffer::Alloc(Size, L"NullReaderWriter::ReadAsync");
-		uint64	  ReadSize = Read(Buffer.GetData(), SourceOffset, Size);
-		Callback(std::move(Buffer), SourceOffset, ReadSize, UserData);
-		return true;
-	}
 
 	// IOWrite
 	virtual uint64 Write(const void* InData, uint64 DestOffset, uint64 WriteSize) override { return WriteSize; }
@@ -391,7 +372,6 @@ struct FDeferredOpenReader : FIOReader
 	FDeferredOpenReader(FOpenCallback InOpenCallback) : OpenCallback(InOpenCallback) {}
 
 	// IOBase
-	virtual void   FlushAll() override { GetOrOpenInner()->FlushAll(); }
 	virtual uint64 GetSize() override { return GetOrOpenInner()->GetSize(); }
 	virtual bool   IsValid() override { return GetOrOpenInner()->IsValid(); }
 	virtual void   Close() override { if (Inner) { Inner->Close(); } }
@@ -402,9 +382,9 @@ struct FDeferredOpenReader : FIOReader
 	{
 		return GetOrOpenInner()->Read(Dest, SourceOffset, ReadSize);
 	}
-	virtual bool ReadAsync(uint64 SourceOffset, uint64 Size, uint64 UserData, IOCallback Callback) override
+	virtual std::unique_ptr<FAsyncReader> CreateAsyncReader(uint32 MaxPipelineDepth = MAX_IO_PIPELINE_DEPTH) override
 	{
-		return GetOrOpenInner()->ReadAsync(SourceOffset, Size, UserData, Callback);
+		return GetOrOpenInner()->CreateAsyncReader(MaxPipelineDepth);
 	}
 
 	FIOReader* GetOrOpenInner()
diff --git a/Engine/Source/Programs/Unsync/Private/UnsyncJupiter.cpp b/Engine/Source/Programs/Unsync/Private/UnsyncJupiter.cpp
--- a/Engine/Source/Programs/Unsync/Private/UnsyncJupiter.cpp
+++ b/Engine/Source/Programs/Unsync/Private/UnsyncJupiter.cpp
@@ -536,6 +536,8 @@ JupiterPush(const FDirectoryManifest& Manifest, const FRemoteDesc& RemoteDesc, c
 			return FError{.Code = -1};
 		}
 
+		std::unique_ptr<FAsyncReader> AsyncReader = File.CreateAsyncReader();
+
 		for (const FGenericBlock& Block : FileManifest.MacroBlocks)
 		{
 			if (bGotError)
@@ -610,7 +612,7 @@ JupiterPush(const FDirectoryManifest& Manifest, const FRemoteDesc& RemoteDesc, c
 					UploadTasks.run(Task);
 				};
 
-			File.ReadAsync(Block.Offset, Block.Size, 0, ReadCallback);
+			AsyncReader->EnqueueRead(Block.Offset, Block.Size, 0, ReadCallback);
 
 			if (TimePointNow() > NextProgressLogTime && ProcessedBytes > 0)
 			{
@@ -623,7 +625,7 @@ JupiterPush(const FDirectoryManifest& Manifest, const FRemoteDesc& RemoteDesc, c
 			}
 		}
 
-		File.FlushAll();
+		AsyncReader->Flush();
 
 		UploadTasks.wait();
 	}
diff --git a/Engine/Source/Programs/Unsync/Private/UnsyncScavenger.cpp b/Engine/Source/Programs/Unsync/Private/UnsyncScavenger.cpp
--- a/Engine/Source/Programs/Unsync/Private/UnsyncScavenger.cpp
+++ b/Engine/Source/Programs/Unsync/Private/UnsyncScavenger.cpp
@@ -456,9 +456,11 @@ BuildTargetFromScavengedData(FIOWriter&						Output,
 				
 			};
 
+			std::unique_ptr<FAsyncReader> AsyncReader = LocalSourceFile.CreateAsyncReader();
+
 			for (const FCopyCommandWithBlockRange& Command : CopyCommands)
 			{
-				LocalSourceFile.ReadAsync(Command.SourceOffset, Command.Size, (uint64)(&Command), ReadCallback);
+				AsyncReader->EnqueueRead(Command.SourceOffset, Command.Size, (uint64)(&Command), ReadCallback);
 				if (bFoundInvalidBlock)
 				{
 					// If an invalid block is encountered in a source file, then assume that is corrupt and should not be used from this point.
@@ -467,7 +469,7 @@ BuildTargetFromScavengedData(FIOWriter&						Output,
 				}
 			}
 
-			LocalSourceFile.FlushAll();
+			AsyncReader->Flush();
 		}
 
 		auto FilterResult =
diff --git a/Engine/Source/Programs/Unsync/Private/UnsyncTarget.cpp b/Engine/Source/Programs/Unsync/Private/UnsyncTarget.cpp
--- a/Engine/Source/Programs/Unsync/Private/UnsyncTarget.cpp
+++ b/Engine/Source/Programs/Unsync/Private/UnsyncTarget.cpp
@@ -306,6 +306,8 @@ BuildTarget(FIOWriter& Output, FIOReader& Source, FIOReader& Base, const FNeedLi
 		static constexpr uint32 MaxActiveLargeRequests							= 2;
 		uint64					ActiveLargeRequestSizes[MaxActiveLargeRequests] = {};
 
+		std::unique_ptr<FAsyncReader> AsyncReader = DataProvider.CreateAsyncReader();
+
 		while (!ReadSchedule.Requests.empty() && !Error)
 		{
 			uint64 BlockIndex = ~0ull;
@@ -383,7 +385,7 @@ BuildTarget(FIOWriter& Output, FIOReader& Source, FIOReader& Base, const FNeedLi
 				AddGlobalProgress(CmdReadSize, ListType);
 			};
 
-			DataProvider.ReadAsync(Block.SourceOffset, Block.Size, 0, ReadCallback);
+			AsyncReader->EnqueueRead(Block.SourceOffset, Block.Size, 0, ReadCallback);
 
 			for (uint32 I = 0; I < MaxActiveLargeRequests; ++I)
 			{
@@ -402,9 +404,7 @@ BuildTarget(FIOWriter& Output, FIOReader& Source, FIOReader& Base, const FNeedLi
 			ProgressLogger.Add(ReadBytes);
 		}
 
-		DataProvider.FlushAll();
-
-		Output.FlushAll();
+		AsyncReader->Flush();
 
 		ProgressLogger.Complete();
 
diff --git a/Engine/Source/Programs/Unsync/Private/UnsyncTest.cpp b/Engine/Source/Programs/Unsync/Private/UnsyncTest.cpp
--- a/Engine/Source/Programs/Unsync/Private/UnsyncTest.cpp
+++ b/Engine/Source/Programs/Unsync/Private/UnsyncTest.cpp
@@ -491,7 +491,6 @@ TestFiles()
 		memset(TempBuffer.get(), 0, TestFileSize);
 
 		TestFile.Read(TempBuffer.get(), 0, TestFileSize);
-		TestFile.FlushAll();
 		VerifyRandomBytes(TempBuffer.get(), TestFileSize, 321);
 	}
 
@@ -502,7 +501,6 @@ TestFiles()
 		UNSYNC_ASSERT(TestFile.IsValid());
 		UNSYNC_ASSERT(TestFile.GetSize() == TestFileSize);
 		TestFile.Read(TempBuffer.get(), 0, TestFileSize);
-		TestFile.FlushAll();
 		VerifyRandomBytes(TempBuffer.get(), TestFileSize, 321);
 	}
 
