unsync - Add wrapper for MD5 hasher and a generic interface for MD5 and Blake3 strong hashing

diff --git a/Engine/Source/Programs/Unsync/Private/UnsyncHash.cpp b/Engine/Source/Programs/Unsync/Private/UnsyncHash.cpp
--- a/Engine/Source/Programs/Unsync/Private/UnsyncHash.cpp
+++ b/Engine/Source/Programs/Unsync/Private/UnsyncHash.cpp
@@ -57,6 +57,40 @@ FBlake3Hasher::Reset()
 	blake3_hasher_reset(Hasher);
 }
 
+
+FMD5Hasher::FMD5Hasher()
+{
+	static_assert(sizeof(State) >= sizeof(md5_context));
+	static_assert(sizeof(State) == sizeof(FMD5Hasher));
+	static_assert(alignof(FMD5Hasher) >= alignof(md5_context));
+
+	md5_context* Hasher = reinterpret_cast<md5_context*>(State);
+	md5_init(Hasher);
+}
+
+void
+FMD5Hasher::Update(const void* Data, size_t Size)
+{
+	md5_context* Hasher = reinterpret_cast<md5_context*>(State);
+	md5_update(Hasher, reinterpret_cast<const uint8*>(Data), Size);
+}
+
+FHash128
+FMD5Hasher::Finalize()
+{
+	FHash128	 Result;
+	md5_context* Hasher = reinterpret_cast<md5_context*>(State);
+	md5_finish(Hasher, Result.Data);
+	return Result;
+}
+
+void
+FMD5Hasher::Reset()
+{
+	md5_context* Hasher = reinterpret_cast<md5_context*>(State);
+	md5_init(Hasher);
+}
+
 template FHash128 HashBlake3Bytes<FHash128>(const uint8* data, uint64 size);
 template FHash160 HashBlake3Bytes<FHash160>(const uint8* data, uint64 size);
 template FHash256 HashBlake3Bytes<FHash256>(const uint8* data, uint64 size);
@@ -229,4 +263,64 @@ ParseHashFromHexString(EStrongHashAlgorithmID StrongHasher, std::string_view Hex
 	return true;
 }
 
+void
+FStrongHasher::Reset()
+{
+	switch (Algorithm)
+	{
+		case EStrongHashAlgorithmID::MD5:
+			MD5.Reset();
+			break;
+		case EStrongHashAlgorithmID::Blake3_128:
+		case EStrongHashAlgorithmID::Blake3_160:
+		case EStrongHashAlgorithmID::Blake3_256:
+			Blake3.Reset();
+			break;
+		default:
+			UNSYNC_FATAL(L"Unsupported hash algorithm ID: %u", uint32(Algorithm));
+	}
+}
+
+void
+FStrongHasher::Update(const void* Data, uint64 Size)
+{
+	switch (Algorithm)
+	{
+		case EStrongHashAlgorithmID::MD5:
+			MD5.Update(Data, Size);
+			break;
+		case EStrongHashAlgorithmID::Blake3_128:
+		case EStrongHashAlgorithmID::Blake3_160:
+		case EStrongHashAlgorithmID::Blake3_256:
+			Blake3.Update(Data, Size);
+			break;
+		default:
+			UNSYNC_FATAL(L"Unsupported hash algorithm ID: %u", uint32(Algorithm));
+	}
+}
+
+FGenericHash
+FStrongHasher::Finalize()
+{
+	FGenericHash Result = {};
+	switch (Algorithm)
+	{
+		case EStrongHashAlgorithmID::MD5:
+			Result = FGenericHash::FromMd5(MD5.Finalize());
+			break;
+		case EStrongHashAlgorithmID::Blake3_128:
+			Result = FGenericHash::FromBlake3_128(ToHash128(Blake3.Finalize()));
+			break;
+		case EStrongHashAlgorithmID::Blake3_160:
+			Result = FGenericHash::FromBlake3_160(ToHash160(Blake3.Finalize()));
+			break;
+		case EStrongHashAlgorithmID::Blake3_256:
+			Result = FGenericHash::FromBlake3_256(Blake3.Finalize());
+			break;
+		default:
+			UNSYNC_FATAL(L"Unsupported hash algorithm ID: %u", uint32(Algorithm));
+	}
+	return Result;
+}
+
 }  // namespace unsync
diff --git a/Engine/Source/Programs/Unsync/Private/UnsyncHash.h b/Engine/Source/Programs/Unsync/Private/UnsyncHash.h
--- a/Engine/Source/Programs/Unsync/Private/UnsyncHash.h
+++ b/Engine/Source/Programs/Unsync/Private/UnsyncHash.h
@@ -312,6 +312,35 @@ private:
 	uint8 State[STATE_SIZE];
 };
 
+class alignas(8) FMD5Hasher
+{
+public:
+	FMD5Hasher();
+
+	void	 Update(const void* Data, size_t Size);
+	FHash128 Finalize();
+	void	 Reset();
+
+private:
+	static constexpr size_t STATE_SIZE = 88;
+	uint8					State[STATE_SIZE];
+};
+
+class FStrongHasher
+{
+public:
+	FStrongHasher(EStrongHashAlgorithmID InAlgorithm) : Algorithm(InAlgorithm) {}
+
+	void		 Reset();
+	void		 Update(const void* Data, uint64 Size);
+	FGenericHash Finalize();
+
+private:
+	const EStrongHashAlgorithmID Algorithm;
+	FBlake3Hasher				 Blake3;
+	FMD5Hasher					 MD5;
+};
+
 FHash128 HashMd5Bytes(const uint8* Data, uint64 Size);
 
 FGenericHash HashBytes(const uint8* Data, uint64 Size, EHashType HashType);
