unsync - Remove OnMacroBlockGenerated callback

diff --git a/Engine/Source/Programs/Unsync/Private/UnsyncChunking.cpp b/Engine/Source/Programs/Unsync/Private/UnsyncChunking.cpp
--- a/Engine/Source/Programs/Unsync/Private/UnsyncChunking.cpp
+++ b/Engine/Source/Programs/Unsync/Private/UnsyncChunking.cpp
@@ -163,14 +163,6 @@ ComputeBlocksVariableT(FIOReader& Reader, const FComputeBlocksParams& Params)
 
 							Task.MacroBlocks.push_back(CurrentMacroBlock);
 
-							if (Params.OnMacroBlockGenerated)
-							{
-								FBufferView BlockView;
-								BlockView.Data = (LastBlockEnd + Block.Size) - CurrentMacroBlock.Size;
-								BlockView.Size = CurrentMacroBlock.Size;
-								Params.OnMacroBlockGenerated(CurrentMacroBlock, SourceInfo, BlockView);
-							}
-
 							// Reset macro block state
 							MacroBlockHasher.Reset();
 							CurrentMacroBlock.Offset += CurrentMacroBlock.Size;
diff --git a/Engine/Source/Programs/Unsync/Private/UnsyncCore.h b/Engine/Source/Programs/Unsync/Private/UnsyncCore.h
--- a/Engine/Source/Programs/Unsync/Private/UnsyncCore.h
+++ b/Engine/Source/Programs/Unsync/Private/UnsyncCore.h
@@ -150,7 +150,6 @@ struct FComputeBlocksParams
 
 	// Callbacks may be called from worker threads
 	FOnBlockGenerated OnBlockGenerated;
-	FOnBlockGenerated OnMacroBlockGenerated;
 
 	bool bAllowThreading = true;
 };
