unsync - Remove locks from FWindowsAsyncFileReader and explicitly document FAsyncReader methods as non-thread-safe

diff --git a/Engine/Source/Programs/Unsync/Private/UnsyncFile.cpp b/Engine/Source/Programs/Unsync/Private/UnsyncFile.cpp
--- a/Engine/Source/Programs/Unsync/Private/UnsyncFile.cpp
+++ b/Engine/Source/Programs/Unsync/Private/UnsyncFile.cpp
@@ -332,12 +332,11 @@ struct FWindowsAsyncFileReader final : FAsyncReader
 	FWindowsAsyncFileReader(FWindowsFile& InReader, uint32 InMaxPipelineDepth);
 
 	virtual ~FWindowsAsyncFileReader();
-	virtual uint64 GetSize() override { return Inner.GetSize(); }
-	virtual bool   IsValid() override { return Inner.IsValid(); }
+	virtual uint64 GetSize() override { return FileSize; }
+	virtual bool   IsValid() override { return !ErrorCode && !bClosed; }
 	virtual bool   EnqueueRead(uint64 SourceOffset, uint64 Size, uint64 UserData, IOCallback Callback);
 	virtual void   Flush() override;
 
-	// NOTE: Expects that Inner.Mutex is locked
 	uint32 CompleteReadCommand(FOverlappedCommand& Cmd);
 
 	FWindowsFile& Inner;
@@ -349,21 +348,33 @@ struct FWindowsAsyncFileReader final : FAsyncReader
 	FOverlappedCommand Commands[MAX_OVERLAPPED_COMMANDS]		 = {};
 	uint64			   NumCommandsIssued						 = 0;
 
-	std::mutex Mutex;
+	HANDLE		 FileHandle = INVALID_HANDLE_VALUE;
+	uint64		 FileSize	= 0;
+	FAtomicError ErrorCode;
+	std::atomic_bool bClosed;
 };
 
 FWindowsAsyncFileReader::FWindowsAsyncFileReader(FWindowsFile& InReader, uint32 InMaxPipelineDepth)
 : Inner(InReader)
+, FileHandle(InReader.FileHandle)
+, FileSize(InReader.GetSize())
 , MaxQueueDepth(InMaxPipelineDepth)
 {
+	UNSYNC_ASSERT(IsReadOnly(InReader.Mode));
+
+	if (!InReader.IsValid())
+	{
+		ErrorCode.Set(SystemError(L"FWindowsAsyncFileReader source file is invalid", InReader.GetError()));
+		InReader.GetError();
+	}
+
 	for (uint32 I = 0; I < MaxQueueDepth; ++I)
 	{
 		OverlappedEvents[I]			  = CreateEvent(nullptr, true, true, nullptr);
 		Commands[I].Overlapped.hEvent = OverlappedEvents[I];
 	}
 
-	std::lock_guard<std::mutex> LockGuard(Inner.Mutex);
-	Inner.AsyncReaders.push_back(this);
+	Inner.AddAsyncReader(this);
 }
 
 FWindowsAsyncFileReader::~FWindowsAsyncFileReader()
@@ -378,8 +389,7 @@ FWindowsAsyncFileReader::~FWindowsAsyncFileReader()
 		}
 	}
 
-	std::lock_guard<std::mutex> LockGuard(Inner.Mutex);
-	Inner.AsyncReaders.erase(std::remove(Inner.AsyncReaders.begin(), Inner.AsyncReaders.end(), this), Inner.AsyncReaders.end());
+	Inner.RemoveAsyncReader(this);
 }
 
 FWindowsFile::FWindowsFile(const FPath& InFilename, EFileMode InMode, uint64 InSize) : Mode(InMode)
@@ -449,6 +459,19 @@ FWindowsFile::~FWindowsFile()
 	UNSYNC_ASSERT(AsyncReaders.empty());
 }
 
+void FWindowsFile::AddAsyncReader(FWindowsAsyncFileReader* Reader)
+{
+	std::lock_guard<std::mutex> LockGuard(Mutex);
+	AsyncReaders.push_back(Reader);
+}
+
+void
+FWindowsFile::RemoveAsyncReader(FWindowsAsyncFileReader* Reader)
+{
+	std::lock_guard<std::mutex> LockGuard(Mutex);
+	AsyncReaders.erase(std::remove(AsyncReaders.begin(), AsyncReaders.end(), Reader), AsyncReaders.end());
+}
+
 bool
 FWindowsFile::OpenFileHandle(EFileMode InMode)
 {
@@ -482,40 +505,45 @@ FWindowsFile::IsValid()
 void
 FWindowsFile::Close()
 {
+	std::lock_guard<std::mutex> LockGuard(Mutex);
+
 	FlushAsyncReaders();
 
-	std::lock_guard<std::mutex> LockGuard(Mutex);
 	if (FileHandle != INVALID_HANDLE_VALUE)
 	{
 		CloseHandle(FileHandle);
 		FileHandle = INVALID_HANDLE_VALUE;
 	}
+
+	for (FWindowsAsyncFileReader* It : AsyncReaders)
+	{
+		It->bClosed = true;
+	}
 }
 
 uint32
 FWindowsAsyncFileReader::CompleteReadCommand(FOverlappedCommand& Cmd)
 {
-	// Expects that Inner.Mutex is locked
-
 	UNSYNC_ASSERT(Cmd.bActive);
 
 	DWORD ReadBytes = 0;
 
-	const uint64 ExpectedReadEnd   = std::min(Inner.FileSize, Cmd.SourceOffset + Cmd.ReadSize);
+	const uint64 ExpectedReadEnd   = std::min(FileSize, Cmd.SourceOffset + Cmd.ReadSize);
 	const uint64 ExpectedReadBytes = ExpectedReadEnd - Cmd.SourceOffset;
 
-	BOOL OverlappedResultOk = GetOverlappedResult(Inner.FileHandle, &Cmd.Overlapped, &ReadBytes, true);
-
-	if (OverlappedResultOk)
 	{
-		Cmd.ErrorCode = 0;
-	}
-	else
-	{
-		Cmd.ErrorCode = GetLastError();
-	}
+		BOOL OverlappedResultOk = GetOverlappedResult(FileHandle, &Cmd.Overlapped, &ReadBytes, true);
 
-	Inner.LastError = Cmd.ErrorCode;
+		if (OverlappedResultOk)
+		{
+			Cmd.ErrorCode = 0;
+		}
+		else
+		{
+			Cmd.ErrorCode = GetLastError();
+			ErrorCode.Set(SystemError(L"GetOverlappedResult failed", Cmd.ErrorCode));
+		}
+	}
 
 	// TODO: retry read command a few times to retry from errors
 	UNSYNC_ASSERT(ReadBytes == ExpectedReadBytes);
@@ -661,11 +689,14 @@ FWindowsFile::Read(void* Dest, uint64 SourceOffset, uint64 ReadSize)
 bool
 FWindowsAsyncFileReader::EnqueueRead(uint64 SourceOffset, uint64 Size, uint64 UserData, IOCallback Callback)
 {
-	std::lock_guard<std::mutex> LockGuard(Inner.Mutex);
+	if (!IsValid())
+	{
+		return false;
+	}
 
-	const EFileMode Mode = Inner.Mode;
+	const EFileMode FileMode = Inner.Mode;
 
-	UNSYNC_ASSERT(IsReadable(Mode));
+	UNSYNC_ASSERT(IsReadable(FileMode));
 
 	uint32 CmdIdx = ~0u;
 
@@ -675,7 +706,6 @@ FWindowsAsyncFileReader::EnqueueRead(uint64 SourceOffset, uint64 Size, uint64 Us
 	UNSYNC_ASSERT(WaitResult == WAIT_OBJECT_0);
 	CmdIdx = WaitSlotIndex;
 
-
 	UNSYNC_ASSERT(CmdIdx < MaxQueueDepth);
 
 	FOverlappedCommand& Cmd = Commands[CmdIdx];
@@ -685,7 +715,7 @@ FWindowsAsyncFileReader::EnqueueRead(uint64 SourceOffset, uint64 Size, uint64 Us
 		CompleteReadCommand(Cmd);
 	}
 
-	if (EnumHasAnyFlags(Mode, EFileMode::Unbuffered))
+	if (EnumHasAnyFlags(FileMode, EFileMode::Unbuffered))
 	{
 		uint64 OriginalSize	 = Size;
 		uint64 OriginalBegin = SourceOffset;
@@ -717,11 +747,20 @@ FWindowsAsyncFileReader::EnqueueRead(uint64 SourceOffset, uint64 Size, uint64 Us
 	Cmd.Callback			  = Callback;
 
 	ResetEvent(Cmd.Overlapped.hEvent);
-	DWORD ReadOk = ReadFile(Inner.FileHandle, Cmd.Buffer.GetMemory(), CheckedNarrow(Size), nullptr, &Cmd.Overlapped);
-	if (!ReadOk)
+
 	{
-		Inner.LastError = GetLastError();
+		DWORD ReadOk = ReadFile(FileHandle, Cmd.Buffer.GetMemory(), CheckedNarrow(Size), nullptr, &Cmd.Overlapped);
+		if (!ReadOk)
+		{
+			const DWORD LastError = GetLastError();
+			if (LastError != ERROR_IO_PENDING)
+			{
+				Cmd.ErrorCode = LastError;
+				ErrorCode.Set(SystemError(L"ReadFile failed", LastError));
+			}
+		}
 	}
+
 	Cmd.bActive = true;
 
 	++NumCommandsIssued;
@@ -732,8 +771,6 @@ FWindowsAsyncFileReader::EnqueueRead(uint64 SourceOffset, uint64 Size, uint64 Us
 void
 FWindowsAsyncFileReader::Flush()
 {
-	std::lock_guard<std::mutex> LockGuard(Inner.Mutex);
-
 	for (uint32 I = 0; I < MaxQueueDepth; ++I)
 	{
 		const uint64 CommandIndex = (NumCommandsIssued + I) % MaxQueueDepth;
@@ -755,6 +792,8 @@ void FWindowsFile::FlushAsyncReaders()
 std::unique_ptr<FAsyncReader>
 FWindowsFile::CreateAsyncReader(uint32 MaxPipelineDepth)
 {
+	UNSYNC_ASSERT(IsValid());
+
 	MaxPipelineDepth = std::min(MaxPipelineDepth, FWindowsAsyncFileReader::MAX_OVERLAPPED_COMMANDS);
 	return std::unique_ptr<FAsyncReader>(new FWindowsAsyncFileReader(*this, MaxPipelineDepth));
 }
diff --git a/Engine/Source/Programs/Unsync/Private/UnsyncFile.h b/Engine/Source/Programs/Unsync/Private/UnsyncFile.h
--- a/Engine/Source/Programs/Unsync/Private/UnsyncFile.h
+++ b/Engine/Source/Programs/Unsync/Private/UnsyncFile.h
@@ -143,8 +143,8 @@ struct FAsyncReader
 	virtual ~FAsyncReader()																			   = default;
 	virtual uint64 GetSize()																		   = 0;
 	virtual bool   IsValid()																		   = 0;
-	virtual bool   EnqueueRead(uint64 SourceOffset, uint64 Size, uint64 UserData, IOCallback Callback) = 0;
-	virtual void   Flush()																			   = 0;
+	virtual bool   EnqueueRead(uint64 SourceOffset, uint64 Size, uint64 UserData, IOCallback Callback) = 0; // NOT thread-safe
+	virtual void   Flush()																			   = 0; // NOT thread-safe
 };
 
 struct FIOReader : virtual FIOBase
@@ -215,7 +215,9 @@ private:
 
 	std::mutex Mutex;
 
-	void FlushAsyncReaders();
+	void								  FlushAsyncReaders();
+	void								  AddAsyncReader(FWindowsAsyncFileReader* Reader);
+	void								  RemoveAsyncReader(FWindowsAsyncFileReader* Reader);
 	std::vector<FWindowsAsyncFileReader*> AsyncReaders;
 };
 using FNativeFile = FWindowsFile;
