unsync - New experimental alternative chunking implementation using streaming through FAsyncReader instead of in-memory chunking

* This significantly reduces memory footprint during hashing (a factor of 10x on a representative data set)
* Allows for increased number of files processed concurrently
* Results in 1.3x overall speedup when computing a manifest for a large game data set
* Off by default initially, opt-in via --streaming command line arg

diff --git a/Engine/Source/Programs/Unsync/Private/UnsyncChunking.cpp b/Engine/Source/Programs/Unsync/Private/UnsyncChunking.cpp
--- a/Engine/Source/Programs/Unsync/Private/UnsyncChunking.cpp
+++ b/Engine/Source/Programs/Unsync/Private/UnsyncChunking.cpp
@@ -32,9 +32,15 @@ ComputeWindowHashThreshold(uint32 TargetSize)
 	return uint32((1ull << 32) / (TargetSize-MinSize));
 }
 
-template<typename WeakHasher>
+void UNSYNC_ATTRIB_NOINLINE
+MemcpyNoInline(void* Dst, const void* Src, size_t Size)
+{
+	memcpy(Dst, Src, Size);
+}
+
+template<typename WeakHasherT>
 FComputeBlocksResult
-ComputeBlocksVariableT(FIOReader& Reader, const FComputeBlocksParams& Params)
+ComputeBlocksVariableStreamingT(FIOReader& Reader, const FComputeBlocksParams& Params)
 {
 	const uint32 BlockSize = Params.BlockSize;
 
@@ -55,16 +61,330 @@ ComputeBlocksVariableT(FIOReader& Reader, const FComputeBlocksParams& Params)
 	FBlockSourceInfo SourceInfo;
 	SourceInfo.TotalSize = InputSize;
 
+	struct FHotState
+	{
+		uint8*		BlockWindowBegin = nullptr;	 // oldest element
+		uint8*		BlockWindowEnd	 = nullptr;	 // newest element
+		WeakHasherT WeakHasher;
+	};
+
 	struct FTask
 	{
 		uint64			   Offset = 0;
 		FGenericBlockArray Blocks;
 		FGenericBlockArray MacroBlocks;
+
+		FBuffer BlockBuffer;
+
+		uint64	  NumBlockBytes = 0;
+
+		FHotState HotState;
 	};
 
 	std::vector<FTask> Tasks;
 	Tasks.resize(NumTasks);
 
+	FTaskGroup			TaskGroup = GScheduler->CreateTaskGroup();
+
+	for (uint64 TaskIndex = 0; TaskIndex < NumTasks; ++TaskIndex)
+	{
+		const uint64 ThisTaskOffset = BytesPerTask * TaskIndex;
+		const uint64 ThisTaskSize	= CalcChunkSize(TaskIndex, BytesPerTask, InputSize);
+
+		{
+			FTask& Task = Tasks[TaskIndex];
+			Task.Offset = ThisTaskOffset;
+			Task.BlockBuffer.Resize(MaximumBlockSize);
+
+			Task.HotState.BlockWindowBegin = Task.BlockBuffer.Data();
+			Task.HotState.BlockWindowEnd   = Task.BlockBuffer.Data();
+		}
+
+		const uint32 ChunkWindowHashThreshold = ComputeWindowHashThreshold(Params.BlockSize);
+
+		auto ScanTask = [&Tasks,
+						 &Params,
+						 &SourceInfo,
+						 &Reader,
+						 MinimumBlockSize,
+						 MaximumBlockSize,
+						 TaskIndex,
+						 ThisTaskSize,
+						 ChunkWindowHashThreshold,
+						 TargetMacroBlockSize,
+						 MinimumMacroBlockSize,
+						 MaximumMacroBlockSize,
+						 MacroBlockHashThreshold]()
+		{
+			FTask& Task = Tasks[TaskIndex];
+
+			const uint64 ReadBatchSize = std::max<uint64>(1_MB, MaximumBlockSize);
+			const uint64 NumReadBatches = DivUp(ThisTaskSize, ReadBatchSize);
+
+			std::unique_ptr<FAsyncReader> AsyncReader = Reader.CreateAsyncReader(8);
+
+			FBlake3Hasher MacroBlockHasher;
+
+			FGenericBlock CurrentMacroBlock;
+			CurrentMacroBlock.HashStrong.Type = MACRO_BLOCK_HASH_TYPE;
+			CurrentMacroBlock.Offset		  = Task.Offset;
+
+			IOCallback ReadCallback = [&Task,
+									   &CurrentMacroBlock,
+									   &Params,
+									   &SourceInfo,
+									   &MacroBlockHasher,
+									   TargetMacroBlockSize,
+									   MinimumMacroBlockSize,
+									   MaximumMacroBlockSize,
+									   MacroBlockHashThreshold,
+									   MinimumBlockSize,
+									   MaximumBlockSize,
+									   ChunkWindowHashThreshold,
+									   NumReadBatches](FIOBuffer Buffer, uint64 SourceOffset, uint64 ReadSize, uint64 UserData) 
+			{
+				uint8* BlockBufferData = Task.BlockBuffer.Data();
+
+				const uint8* DataBegin = Buffer.GetData();
+				const uint8* DataEnd   = DataBegin + ReadSize;
+
+				const uint8* Cursor = DataBegin;
+
+				const bool bLastBatch = UserData + 1 == NumReadBatches;
+
+				FHotState State = Task.HotState;
+
+				while (Cursor != DataEnd)
+				{
+					while (State.WeakHasher.Count < MinimumBlockSize && Cursor != DataEnd)
+					{
+						State.WeakHasher.Add(*Cursor);
+
+						(*State.BlockWindowEnd++) = *Cursor;
+
+						++Cursor;
+					}
+
+					const uint32 WindowHash = State.WeakHasher.Get();
+
+					const bool	 bLastBlock	   = Cursor == DataEnd && bLastBatch;
+					const uint64 ThisBlockSize = uint64(State.BlockWindowEnd - BlockBufferData);
+
+					if (ThisBlockSize >= MaximumBlockSize || WindowHash < ChunkWindowHashThreshold || bLastBlock)
+					{
+						FGenericBlock Block;
+						Block.Offset	 = Task.Offset + Task.NumBlockBytes;
+						Block.Size		 = CheckedNarrow(ThisBlockSize);
+						Block.HashWeak	 = WindowHash;
+						Block.HashStrong = ComputeHash(BlockBufferData, ThisBlockSize, Params.Algorithm.StrongHashAlgorithmId);
+
+						if (TargetMacroBlockSize)
+						{
+							MacroBlockHasher.Update(BlockBufferData, ThisBlockSize);
+							CurrentMacroBlock.Size += Block.Size;
+
+							uint32 HashStrong32 = 0;
+							memcpy(&HashStrong32, Block.HashStrong.Data, 4);
+
+							if ((CurrentMacroBlock.Size >= MinimumMacroBlockSize && HashStrong32 < MacroBlockHashThreshold) ||
+								(CurrentMacroBlock.Size + Block.Size > MaximumMacroBlockSize) || bLastBlock)
+							{
+								// Commit the macro block
+								const FHash256 MacroBlockHash	  = MacroBlockHasher.Finalize();
+								CurrentMacroBlock.HashStrong	  = FGenericHash::FromBlake3_256(MacroBlockHash);
+								CurrentMacroBlock.HashStrong.Type = MACRO_BLOCK_HASH_TYPE;
+
+								Task.MacroBlocks.push_back(CurrentMacroBlock);
+
+								// Reset macro block state
+								MacroBlockHasher.Reset();
+								CurrentMacroBlock.Offset += CurrentMacroBlock.Size;
+								CurrentMacroBlock.Size = 0;
+							}
+						}
+
+						if (Params.OnBlockGenerated)
+						{
+							FBufferView BlockView;
+							BlockView.Data = BlockBufferData;
+							BlockView.Size = Block.Size;
+							Params.OnBlockGenerated(Block, SourceInfo, BlockView);
+						}
+
+						if (!Task.Blocks.empty())
+						{
+							UNSYNC_ASSERT(Task.Blocks.back().Offset + Task.Blocks.back().Size == Block.Offset);
+						}
+
+						Task.NumBlockBytes += Block.Size;
+
+						Task.Blocks.push_back(Block);
+
+						State.WeakHasher.Reset();
+
+						State.BlockWindowBegin = BlockBufferData;
+						State.BlockWindowEnd   = BlockBufferData;
+
+						continue;
+					}
+
+					State.WeakHasher.Sub(*(State.BlockWindowBegin++));
+				}
+
+				//memcpy(&Task.HotState, &HotState, sizeof(HotState));
+				MemcpyNoInline(&Task.HotState, &State, sizeof(State)); // NOTE: no-inline memcpy generates MUCH faster chunking loop code
+
+			};
+
+			for (uint64 BatchIndex = 0; BatchIndex < NumReadBatches; ++BatchIndex)
+			{
+				const uint64 ThisReadSize = CalcChunkSize(BatchIndex, ReadBatchSize, ThisTaskSize);
+				AsyncReader->EnqueueRead(Task.Offset + BatchIndex * ReadBatchSize, ThisReadSize, BatchIndex, ReadCallback);
+			}
+
+		};
+
+		TaskGroup.run(ScanTask);
+	}
+
+	TaskGroup.wait();
+
+	// Merge blocks for all the tasks
+
+	FComputeBlocksResult Result;
+
+	for (uint64 I = 0; I < NumTasks; ++I)
+	{
+		const FTask& Task = Tasks[I];
+		for (uint64 J = 0; J < Task.Blocks.size(); ++J)
+		{
+			Result.Blocks.push_back(Task.Blocks[J]);
+		}
+	}
+
+	if (Params.bNeedMacroBlocks)
+	{
+		for (uint64 I = 0; I < NumTasks; ++I)
+		{
+			const FTask& Task = Tasks[I];
+			for (uint64 J = 0; J < Task.MacroBlocks.size(); ++J)
+			{
+				Result.MacroBlocks.push_back(Task.MacroBlocks[J]);
+			}
+		}
+	}
+
+	uint64 UniqueBlockTotalSize = 0;
+	uint64 UniqueBlockMinSize	= ~0ull;
+	uint64 UniqueBlockMaxSize	= 0ull;
+
+	uint64 NumTinyBlocks   = 0;
+	uint64 NumSmallBlocks  = 0;
+	uint64 NumMediumBlocks = 0;
+	uint64 NumLargeBlocks  = 0;
+
+	uint64 NumTotalBlocks = 0;
+
+	THashSet<FGenericHash> UniqueBlockSet;
+	FGenericBlockArray	   UniqueBlocks;
+	for (const FGenericBlock& It : Result.Blocks)
+	{
+		auto InsertResult = UniqueBlockSet.insert(It.HashStrong);
+		if (InsertResult.second)
+		{
+			if (It.Offset + It.Size < InputSize || Result.Blocks.size() == 1)
+			{
+				UniqueBlockMinSize = std::min<uint64>(UniqueBlockMinSize, It.Size);
+			}
+
+			UniqueBlockMaxSize = std::max<uint64>(UniqueBlockMaxSize, It.Size);
+			UniqueBlockTotalSize += It.Size;
+			UniqueBlocks.push_back(It);
+		}
+
+		if (It.Size < MaximumBlockSize / 8)
+		{
+			NumTinyBlocks++;
+		}
+		else if (It.Size <= MaximumBlockSize / 4)
+		{
+			NumSmallBlocks++;
+		}
+		else if (It.Size <= MaximumBlockSize / 2)
+		{
+			NumMediumBlocks++;
+		}
+		else
+		{
+			NumLargeBlocks++;
+		}
+
+		++NumTotalBlocks;
+	}
+
+	double AverageBlockSize = InputSize ? double(UniqueBlockTotalSize / UniqueBlocks.size()) : 0;
+
+	UNSYNC_VERBOSE2(
+		L"Blocks (tiny/small/medium/large): %llu / %llu / %llu / %llu, average unique size: %llu bytes, unique count: %llu, total count: "
+		L"%llu",
+		NumTinyBlocks,
+		NumSmallBlocks,
+		NumMediumBlocks,
+		NumLargeBlocks,
+		(uint64)AverageBlockSize,
+		(uint64)UniqueBlocks.size(),
+		NumTotalBlocks);
+
+	UNSYNC_ASSERT(NumTotalBlocks == Result.Blocks.size());
+
+	return Result;
+}
+
+template<typename WeakHasherT>
+FComputeBlocksResult
+ComputeBlocksVariableT(FIOReader& Reader, const FComputeBlocksParams& Params)
+{
+	const uint32 BlockSize = Params.BlockSize;
+
+	const uint64 InputSize = Reader.GetSize();
+
+	const uint32 MinimumBlockSize = ComputeMinVariableBlockSize(BlockSize);
+	const uint32 MaximumBlockSize = ComputeMaxVariableBlockSize(BlockSize);
+
+	const uint64 BytesPerTask = std::min<uint64>(256_MB, std::max<uint64>(InputSize, 1ull));  // TODO: handle task boundary overlap
+	const uint64 NumTasks	  = DivUp(InputSize, BytesPerTask);
+
+	const uint64 TargetMacroBlockSize	 = Params.bNeedMacroBlocks ? Params.MacroBlockTargetSize : 0;
+	const uint64 MinimumMacroBlockSize	 = std::max<uint64>(MinimumBlockSize, TargetMacroBlockSize / 8);
+	const uint64 MaximumMacroBlockSize	 = Params.bNeedMacroBlocks ? Params.MacroBlockMaxSize : 0;
+	const uint32 BlocksPerMacroBlock	 = CheckedNarrow(DivUp(TargetMacroBlockSize - MinimumMacroBlockSize, BlockSize));
+	const uint32 MacroBlockHashThreshold = BlocksPerMacroBlock ? (0xFFFFFFFF / BlocksPerMacroBlock) : 0;
+
+	FBlockSourceInfo SourceInfo;
+	SourceInfo.TotalSize = InputSize;
+
+	struct FTask
+	{
+		FTask(EStrongHashAlgorithmID Algorithm)
+			: StrongHasher(Algorithm) 
+		{
+		}
+
+		uint64			   Offset = 0;
+		FGenericBlockArray Blocks;
+		FGenericBlockArray MacroBlocks;
+
+		FStrongHasher StrongHasher;
+		FBlake3Hasher MacroHasher;
+	};
+
+	std::vector<FTask> Tasks;
+	Tasks.reserve(NumTasks);
+	for (uint64 TaskIndex = 0; TaskIndex < NumTasks; ++TaskIndex)
+	{
+		Tasks.emplace_back(Params.Algorithm.StrongHashAlgorithmId);
+	}
+
 	FSchedulerSemaphore IoSemaphore(*GScheduler, 16);
 	FTaskGroup			TaskGroup = GScheduler->CreateTaskGroup();
 
@@ -109,16 +429,14 @@ ComputeBlocksVariableT(FIOReader& Reader, const FComputeBlocksParams& Params)
 		{
 			FTask& Task = Tasks[TaskIndex];
 
-			const uint8* DataBegin	  = ScanTaskBuffer->Data();
-			const uint8* DataEnd	  = DataBegin + ThisTaskSize;
-			const uint8* LastBlockEnd = DataBegin;
+			const uint8* DataBegin = ScanTaskBuffer->Data();
+			const uint8* DataEnd   = DataBegin + ThisTaskSize;
 
 			FGenericBlock CurrentMacroBlock;
 			CurrentMacroBlock.HashStrong.Type = MACRO_BLOCK_HASH_TYPE;
 			CurrentMacroBlock.Offset		  = Task.Offset;
 
-			FBlake3Hasher MacroBlockHasher;
-
+			const uint8* LastBlockEnd = DataBegin;
 			auto ScanFn = [MacroBlockHashThreshold,
 						   MaximumMacroBlockSize,
 						   MaximumBlockSize,
@@ -130,7 +448,6 @@ ComputeBlocksVariableT(FIOReader& Reader, const FComputeBlocksParams& Params)
 						   DataBegin,
 						   &Params,
 						   TargetMacroBlockSize,
-						   &MacroBlockHasher,
 						   &CurrentMacroBlock,
 						   &SourceInfo](const uint8* WindowBegin, const uint8* WindowEnd, uint32 WindowHash) UNSYNC_ATTRIB_FORCEINLINE
 			{
@@ -147,7 +464,7 @@ ComputeBlocksVariableT(FIOReader& Reader, const FComputeBlocksParams& Params)
 
 					if (TargetMacroBlockSize)
 					{
-						MacroBlockHasher.Update(LastBlockEnd, ThisBlockSize);
+						Task.MacroHasher.Update(LastBlockEnd, ThisBlockSize);
 						CurrentMacroBlock.Size += Block.Size;
 
 						uint32 HashStrong32 = 0;
@@ -157,14 +474,14 @@ ComputeBlocksVariableT(FIOReader& Reader, const FComputeBlocksParams& Params)
 							(CurrentMacroBlock.Size + Block.Size > MaximumMacroBlockSize) || bLastBlock)
 						{
 							// Commit the macro block
-							const FHash256 MacroBlockHash	  = MacroBlockHasher.Finalize();
+							const FHash256 MacroBlockHash	  = Task.MacroHasher.Finalize();
 							CurrentMacroBlock.HashStrong	  = FGenericHash::FromBlake3_256(MacroBlockHash);
 							CurrentMacroBlock.HashStrong.Type = MACRO_BLOCK_HASH_TYPE;
 
 							Task.MacroBlocks.push_back(CurrentMacroBlock);
 
 							// Reset macro block state
-							MacroBlockHasher.Reset();
+							Task.MacroHasher.Reset();
 							CurrentMacroBlock.Offset += CurrentMacroBlock.Size;
 							CurrentMacroBlock.Size = 0;
 						}
@@ -195,7 +512,7 @@ ComputeBlocksVariableT(FIOReader& Reader, const FComputeBlocksParams& Params)
 				}
 			};
 
-			HashScan<WeakHasher>(DataBegin, ThisTaskSize, MinimumBlockSize, ScanFn);
+			HashScan<WeakHasherT>(DataBegin, ThisTaskSize, MinimumBlockSize, ScanFn);
 
 			BufferPool.Release(ScanTaskBuffer);
 			IoSemaphore.Release();
@@ -300,15 +617,31 @@ ComputeBlocksVariableT(FIOReader& Reader, const FComputeBlocksParams& Params)
 FComputeBlocksResult
 ComputeBlocksVariable(FIOReader& Reader, const FComputeBlocksParams& Params)
 {
-	switch (Params.Algorithm.WeakHashAlgorithmId)
+	if (Params.bAllowStreaming)
 	{
-		case EWeakHashAlgorithmID::Naive:
-			return ComputeBlocksVariableT<FRollingChecksum>(Reader, Params);
-		case EWeakHashAlgorithmID::BuzHash:
-			return ComputeBlocksVariableT<FBuzHash>(Reader, Params);
-		default:
-			UNSYNC_FATAL(L"Unsupported weak hash algorithm mode");
-			return {};
+		switch (Params.Algorithm.WeakHashAlgorithmId)
+		{
+			case EWeakHashAlgorithmID::Naive:
+				return ComputeBlocksVariableStreamingT<FRollingChecksum>(Reader, Params);
+			case EWeakHashAlgorithmID::BuzHash:
+				return ComputeBlocksVariableStreamingT<FBuzHash>(Reader, Params);
+			default:
+				UNSYNC_FATAL(L"Unsupported weak hash algorithm mode");
+				return {};
+		}
+	}
+	else
+	{
+		switch (Params.Algorithm.WeakHashAlgorithmId)
+		{
+			case EWeakHashAlgorithmID::Naive:
+				return ComputeBlocksVariableT<FRollingChecksum>(Reader, Params);
+			case EWeakHashAlgorithmID::BuzHash:
+				return ComputeBlocksVariableT<FBuzHash>(Reader, Params);
+			default:
+				UNSYNC_FATAL(L"Unsupported weak hash algorithm mode");
+				return {};
+		}
 	}
 }
 
diff --git a/Engine/Source/Programs/Unsync/Private/UnsyncCommon.h b/Engine/Source/Programs/Unsync/Private/UnsyncCommon.h
--- a/Engine/Source/Programs/Unsync/Private/UnsyncCommon.h
+++ b/Engine/Source/Programs/Unsync/Private/UnsyncCommon.h
@@ -17,6 +17,7 @@
 #if UNSYNC_COMPILER_MSVC
 #	pragma warning(disable : 4100)	 // unreferenced formal parameter
 #	define UNSYNC_ATTRIB_FORCEINLINE [[msvc::forceinline]]
+#	define UNSYNC_ATTRIB_NOINLINE	  __declspec(noinline)
 #	define UNSYNC_THIRD_PARTY_INCLUDES_START \
 		__pragma(warning(push)) \
 		__pragma(warning(disable : 4458)) \
@@ -25,6 +26,7 @@
 #	define UNSYNC_THIRD_PARTY_INCLUDES_END __pragma(warning(pop))
 #else
 #	define UNSYNC_ATTRIB_FORCEINLINE __attribute__((always_inline))
+#	define UNSYNC_ATTRIB_NOINLINE	  __attribute__((noinline))
 #	define UNSYNC_THIRD_PARTY_INCLUDES_START \
 		_Pragma("GCC diagnostic push") \
 		_Pragma("GCC diagnostic ignored \"-Wdeprecated-declarations\"") \
diff --git a/Engine/Source/Programs/Unsync/Private/UnsyncCore.cpp b/Engine/Source/Programs/Unsync/Private/UnsyncCore.cpp
--- a/Engine/Source/Programs/Unsync/Private/UnsyncCore.cpp
+++ b/Engine/Source/Programs/Unsync/Private/UnsyncCore.cpp
@@ -40,6 +40,7 @@ namespace unsync {
 
 bool GDryRun = false;
 bool GExperimental = false;
+bool GExperimentalStreaming = false;
 
 FBlock128
 ToBlock128(const FGenericBlock& GenericBlock)
diff --git a/Engine/Source/Programs/Unsync/Private/UnsyncCore.h b/Engine/Source/Programs/Unsync/Private/UnsyncCore.h
--- a/Engine/Source/Programs/Unsync/Private/UnsyncCore.h
+++ b/Engine/Source/Programs/Unsync/Private/UnsyncCore.h
@@ -22,6 +22,7 @@ namespace unsync {
 
 extern bool GDryRun;
 extern bool GExperimental;
+extern bool GExperimentalStreaming;
 
 class FProxy;
 class FProxyPool;
@@ -152,6 +153,8 @@ struct FComputeBlocksParams
 	FOnBlockGenerated OnBlockGenerated;
 
 	bool bAllowThreading = true;
+
+	bool bAllowStreaming = GExperimentalStreaming;
 };
 
 
diff --git a/Engine/Source/Programs/Unsync/Private/UnsyncMain.cpp b/Engine/Source/Programs/Unsync/Private/UnsyncMain.cpp
--- a/Engine/Source/Programs/Unsync/Private/UnsyncMain.cpp
+++ b/Engine/Source/Programs/Unsync/Private/UnsyncMain.cpp
@@ -409,6 +409,7 @@ InnerMain(int Argc, char** Argv)
 		Subcommand->add_flag("--buffered-files", GForceBufferedFiles, "Always use buffered file IO");
 		Subcommand->add_flag("--debug", bUseDebugMode, "Enable extra debugging features, such as extra memory safety validation");
 		Subcommand->add_flag("--experimental", GExperimental, "Enable experimental code paths")->group(HiddenGroupId);
+		Subcommand->add_flag("--streaming", GExperimentalStreaming, "Enable experimental file streaming code path")->group(HiddenGroupId);
 
 		SilentFlag->excludes(VerboseFlag);
 		SilentFlag->excludes(VeryVerboseFlag);
diff --git a/Engine/Source/Programs/Unsync/Private/UnsyncManifest.cpp b/Engine/Source/Programs/Unsync/Private/UnsyncManifest.cpp
--- a/Engine/Source/Programs/Unsync/Private/UnsyncManifest.cpp
+++ b/Engine/Source/Programs/Unsync/Private/UnsyncManifest.cpp
@@ -174,7 +174,7 @@ UpdateDirectoryManifestBlocks(FDirectoryManifest& Result, const FPath& Root, con
 
 	uint32 NumProcessedFiles = 0;
 
-	const uint32		MaxConcurrentFiles = 8;	 // quickly diminishing returns past 8 concurrent files
+	const uint32		MaxConcurrentFiles = Params.bAllowStreaming ? 16 : 8;
 	FSchedulerSemaphore Semaphore(*GScheduler, MaxConcurrentFiles);
 	FTaskGroup			TaskGroup = GScheduler->CreateTaskGroup(&Semaphore);
 
@@ -280,7 +280,7 @@ CreateDirectoryManifest(const FPath& Root, const FComputeBlocksParams& Params)
 
 	FTimePoint TimeBegin = TimePointNow();
 
-	const uint32		MaxConcurrentFiles = 8;	 // quickly diminishing returns past 8 concurrent files
+	const uint32		MaxConcurrentFiles = Params.bAllowStreaming ? 16 : 8;
 	FSchedulerSemaphore Semaphore(*GScheduler, MaxConcurrentFiles);
 	FTaskGroup			TaskGroup = GScheduler->CreateTaskGroup(&Semaphore);
 
