unsync - Use _wgetenv to get user home directory on Windows

diff --git a/Engine/Source/Programs/Unsync/Private/UnsyncUtil.cpp b/Engine/Source/Programs/Unsync/Private/UnsyncUtil.cpp
--- a/Engine/Source/Programs/Unsync/Private/UnsyncUtil.cpp
+++ b/Engine/Source/Programs/Unsync/Private/UnsyncUtil.cpp
@@ -495,15 +495,52 @@ GetUniversalPath(const FPath& Path)
 	return Result;
 }
 
+static FPath
+GetNormalWeaklyCanonicalAbsolutePath(const FPath& InPath)
+{
+	FPath NormalPath		 = InPath.lexically_normal();
+	FPath CanonicalPath		 = std::filesystem::weakly_canonical(NormalPath);
+	FPath AbsoluteNormalPath = std::filesystem::absolute(CanonicalPath);
+	return AbsoluteNormalPath;
+}
+
+FPath
+NormalizeFilenameWide(std::wstring_view Filename)
+{
+	if (Filename.empty())
+	{
+		return FPath();
+	}
+
+	std::wstring_view FileUrlPrefix = L"file://";
+	if (Filename.starts_with(FileUrlPrefix))
+	{
+		Filename = Filename.substr(FileUrlPrefix.length());
+	}
+
+	FPath FilenameAsPath = FPath(Filename);
+
+	FPath AbsoluteNormalPath;
+	if (Filename.starts_with(L"\\\\") || Filename.starts_with(L"//"))
+	{
+		AbsoluteNormalPath = FilenameAsPath;  // Assume network paths are absolute
+	}
+	else
+	{
+		AbsoluteNormalPath = GetNormalWeaklyCanonicalAbsolutePath(FilenameAsPath);
+	}
+
+	return AbsoluteNormalPath;
+}
+
 FPath
-NormalizeFilenameUtf8(const std::string& InFilename)
+NormalizeFilenameUtf8(std::string_view Filename)
 {
-	if (InFilename.empty())
+	if (Filename.empty())
 	{
 		return FPath();
 	}
 
-	std::string_view Filename	   = InFilename;
 	std::string_view FileUrlPrefix = "file://";
 	if (Filename.starts_with(FileUrlPrefix))
 	{
@@ -519,9 +556,7 @@ NormalizeFilenameUtf8(const std::string& InFilename)
 	}
 	else
 	{
-		FPath NormalPath	= FilenameAsPath.lexically_normal();
-		FPath CanonicalPath = std::filesystem::weakly_canonical(NormalPath);
-		AbsoluteNormalPath	= std::filesystem::absolute(CanonicalPath);
+		AbsoluteNormalPath = GetNormalWeaklyCanonicalAbsolutePath(FilenameAsPath);
 	}
 
 	return AbsoluteNormalPath;
@@ -655,9 +690,9 @@ OpenUrlInDefaultBrowser(const char* Address)
 #if UNSYNC_PLATFORM_WINDOWS
 FPath GetUserHomeDirectory()
 {
-	if (const char* EnvUserProfile = getenv("USERPROFILE"))
+	if (const wchar_t* EnvUserProfile = _wgetenv(L"USERPROFILE"))
 	{
-		return NormalizeFilenameUtf8(EnvUserProfile);
+		return NormalizeFilenameWide(EnvUserProfile);
 	}
 	else
 	{
diff --git a/Engine/Source/Programs/Unsync/Private/UnsyncUtil.h b/Engine/Source/Programs/Unsync/Private/UnsyncUtil.h
--- a/Engine/Source/Programs/Unsync/Private/UnsyncUtil.h
+++ b/Engine/Source/Programs/Unsync/Private/UnsyncUtil.h
@@ -367,7 +367,8 @@ BitArraySet(StorageT* Storage, uint64 BitIndex, bool bValue)
 	}
 }
 
-FPath NormalizeFilenameUtf8(const std::string& InFilename);
+FPath NormalizeFilenameWide(std::wstring_view InFilename);
+FPath NormalizeFilenameUtf8(std::string_view InFilename);
 FPath GetAbsoluteNormalPath(const FPath& InPath);
 
 const FBuffer& GetSystemRootCerts();
diff --git a/Engine/Source/Programs/Unsync/Private/UnsyncVersion.cpp b/Engine/Source/Programs/Unsync/Private/UnsyncVersion.cpp
--- a/Engine/Source/Programs/Unsync/Private/UnsyncVersion.cpp
+++ b/Engine/Source/Programs/Unsync/Private/UnsyncVersion.cpp
@@ -2,7 +2,7 @@
 
 #include "UnsyncVersion.h"
 
-#define UNSYNC_VERSION_STR "1.0.93"
+#define UNSYNC_VERSION_STR "1.0.94"
 
 namespace unsync {
 
