unsync - Add error context message in QueryListDirectory()

diff --git a/Engine/Source/Programs/Unsync/Private/UnsyncProxy.cpp b/Engine/Source/Programs/Unsync/Private/UnsyncProxy.cpp
--- a/Engine/Source/Programs/Unsync/Private/UnsyncProxy.cpp
+++ b/Engine/Source/Programs/Unsync/Private/UnsyncProxy.cpp
@@ -956,7 +956,8 @@ FUnsyncProtocolImpl::QueryListDirectory(FHttpConnection& Connection, const FAuth
 
 	if (!Response.Success())
 	{
-		return HttpError(Response.Code);
+		std::string ErrorContext = fmt::format("Listing server directory: {}", Path);
+		return HttpError(std::move(ErrorContext), Response.Code);
 	}
 
 	Response.Buffer.PushBack(0);
