unsync - Add a dedicated command to check a directory against its manifest

diff --git a/Engine/Source/Programs/Unsync/CMakeLists.txt b/Engine/Source/Programs/Unsync/CMakeLists.txt
--- a/Engine/Source/Programs/Unsync/CMakeLists.txt
+++ b/Engine/Source/Programs/Unsync/CMakeLists.txt
@@ -56,6 +56,8 @@ set(src
 	Private/UnsyncCmdQuery.h
 	Private/UnsyncCmdSync.cpp
 	Private/UnsyncCmdSync.h
+	Private/UnsyncCmdVerify.cpp
+	Private/UnsyncCmdVerify.h
 	Private/UnsyncCommon.h
 	Private/UnsyncCompression.cpp
 	Private/UnsyncCompression.h
diff --git a/Engine/Source/Programs/Unsync/Private/UnsyncCmdVerify.cpp b/Engine/Source/Programs/Unsync/Private/UnsyncCmdVerify.cpp
new file mode 100644
--- /dev/null
+++ b/Engine/Source/Programs/Unsync/Private/UnsyncCmdVerify.cpp
@@ -0,0 +1,103 @@
+// Copyright Epic Games, Inc. All Rights Reserved.
+
+#include "UnsyncCmdVerify.h"
+#include "UnsyncFile.h"
+#include "UnsyncLog.h"
+#include "UnsyncManifest.h"
+#include "UnsyncSerialization.h"
+
+namespace unsync {
+
+int32
+CmdVerify(const FCmdVerifyOptions& Options)
+{
+	FPath			   Root					 = Options.Input;
+	FPath			   ManifestRoot			 = Root / ".unsync";
+	FPath			   DirectoryManifestPath = ManifestRoot / "manifest.bin";
+	FDirectoryManifest DirectoryManifest;
+
+	UNSYNC_LOG(L"Validating manifest for directory '%ls'", Root.wstring().c_str());
+	UNSYNC_LOG_INDENT;
+
+	if (!PathExists(DirectoryManifestPath))
+	{
+		UNSYNC_ERROR(L"Directory does not have a manifest file");
+		return -1;
+	}
+
+	const bool bExistingManifestLoaded = LoadDirectoryManifest(DirectoryManifest, Root, DirectoryManifestPath);
+	if (!bExistingManifestLoaded)
+	{
+		UNSYNC_ERROR(L"Failed to load manifest file '%ls'", DirectoryManifestPath.wstring().c_str());
+		return -1;
+	}
+
+	const EStrongHashAlgorithmID StrongHashAlgorithmId = DirectoryManifest.Algorithm.StrongHashAlgorithmId;
+
+	bool bDirectoryValid = true;
+
+	for (const auto& FileIt : DirectoryManifest.Files)
+	{
+		const std::wstring&	 FileName	  = FileIt.first;
+		const FFileManifest& FileManifest = FileIt.second;
+
+		UNSYNC_VERBOSE(L"Verifying file '%ls'", FileName.c_str());
+
+		const FFileAttributes FileAttrib = GetFileAttrib(FileManifest.CurrentPath);
+		if (!FileAttrib.bValid)
+		{
+			bDirectoryValid = false;
+			UNSYNC_ERROR(L"File '%ls' does not exist", FileName.c_str());
+			continue;
+		}
+
+		if (FileAttrib.Size != FileManifest.Size)
+		{
+			bDirectoryValid = false;
+			UNSYNC_ERROR(L"File '%ls' size mismatch. Expected %llu, actual %llu.",
+						 FileName.c_str(),
+						 llu(FileManifest.Size),
+						 llu(FileAttrib.Size));
+			continue;
+		}
+
+		if (FileAttrib.Mtime != FileManifest.Mtime)
+		{
+			bDirectoryValid = false;
+			UNSYNC_ERROR(L"File '%ls' timestamp mismatch. Expected %llu, actual %llu.",
+						 FileName.c_str(),
+						 llu(FileManifest.Mtime),
+						 llu(FileAttrib.Mtime));
+			continue;
+		}
+
+		FNativeFile File(FileManifest.CurrentPath, EFileMode::ReadOnlyUnbuffered);
+		if (!File.IsValid())
+		{
+			bDirectoryValid = false;
+			UNSYNC_ERROR(L"Failed to open file '%ls'. Error code: %d", FileName.c_str(), File.GetError());
+			continue;
+		}
+
+		const bool bFileValid = ValidateTarget(File, FileManifest.Blocks, StrongHashAlgorithmId);
+
+		if (!bFileValid)
+		{
+			UNSYNC_ERROR(L"Validation failed for file '%ls'", FileName.c_str());
+			bDirectoryValid = false;
+		}
+	}
+
+	if (bDirectoryValid)
+	{
+		UNSYNC_LOG(L"Directory manifest is valid");
+		return 0;
+	}
+	else
+	{
+		UNSYNC_LOG(L"Directory manifest is invalid");
+		return -1;
+	}
+}
+
+}  // namespace unsync
diff --git a/Engine/Source/Programs/Unsync/Private/UnsyncCmdVerify.h b/Engine/Source/Programs/Unsync/Private/UnsyncCmdVerify.h
new file mode 100644
--- /dev/null
+++ b/Engine/Source/Programs/Unsync/Private/UnsyncCmdVerify.h
@@ -0,0 +1,14 @@
+// Copyright Epic Games, Inc. All Rights Reserved.
+
+#include "UnsyncCommon.h"
+
+namespace unsync {
+
+struct FCmdVerifyOptions
+{
+	FPath			   Input;
+};
+
+int32 CmdVerify(const FCmdVerifyOptions& Options);
+
+}  // namespace unsync
diff --git a/Engine/Source/Programs/Unsync/Private/UnsyncCore.cpp b/Engine/Source/Programs/Unsync/Private/UnsyncCore.cpp
--- a/Engine/Source/Programs/Unsync/Private/UnsyncCore.cpp
+++ b/Engine/Source/Programs/Unsync/Private/UnsyncCore.cpp
@@ -184,10 +184,16 @@ ValidateTarget(FIOReader& Reader, const FNeedList& NeedList, EStrongHashAlgorith
 		ValidationBlocks.push_back(Block);
 	}
 
-	std::sort(ValidationBlocks.begin(), ValidationBlocks.end(), [](const FGenericBlock& A, const FGenericBlock& B) {
-		return A.Offset < B.Offset;
-	});
+	std::sort(ValidationBlocks.begin(),
+			  ValidationBlocks.end(),
+			  [](const FGenericBlock& A, const FGenericBlock& B) { return A.Offset < B.Offset; });
+
+	return ValidateTarget(Reader, ValidationBlocks, StrongHasher);
+}
 
+bool
+ValidateTarget(FIOReader& Reader, const FGenericBlockArray& ValidationBlocks, EStrongHashAlgorithmID StrongHasher)
+{
 	const uint64		TotalStreamBytes = Reader.GetSize();
 	std::atomic<uint64> NumInvalidBlocks = {};
 	FSchedulerSemaphore IoSemaphore(*GScheduler, 16);
diff --git a/Engine/Source/Programs/Unsync/Private/UnsyncCore.h b/Engine/Source/Programs/Unsync/Private/UnsyncCore.h
--- a/Engine/Source/Programs/Unsync/Private/UnsyncCore.h
+++ b/Engine/Source/Programs/Unsync/Private/UnsyncCore.h
@@ -278,4 +278,7 @@ ValidateBlockListT(const std::vector<BlockType>& Blocks)
 	return true;
 }
 
+bool ValidateTarget(FIOReader& Reader, const FNeedList& NeedList, EStrongHashAlgorithmID StrongHasher);
+bool ValidateTarget(FIOReader& Reader, const FGenericBlockArray& ValidationBlocks, EStrongHashAlgorithmID StrongHasher);
+
 }  // namespace unsync
diff --git a/Engine/Source/Programs/Unsync/Private/UnsyncFile.cpp b/Engine/Source/Programs/Unsync/Private/UnsyncFile.cpp
--- a/Engine/Source/Programs/Unsync/Private/UnsyncFile.cpp
+++ b/Engine/Source/Programs/Unsync/Private/UnsyncFile.cpp
@@ -546,7 +546,14 @@ FWindowsAsyncFileReader::CompleteReadCommand(FOverlappedCommand& Cmd)
 	}
 
 	// TODO: retry read command a few times to retry from errors
-	UNSYNC_ASSERT(ReadBytes == ExpectedReadBytes);
+	UNSYNC_ASSERTF(ReadBytes == ExpectedReadBytes,
+				   L"Expected to read %llu bytes, but read %llu [FileSize=%llu, Cmd.SourceOffset=%llu, Cmd.ReadSize=%llu, Cmd.ErrorCode=%u]",
+				   llu(ExpectedReadBytes),
+				   llu(ReadBytes),
+				   llu(FileSize),
+				   llu(Cmd.SourceOffset),
+				   llu(Cmd.ReadSize),
+				   Cmd.ErrorCode);
 
 	const uint64 ReadBytesClamped = std::min<uint64>(Cmd.Buffer.GetSize(), ReadBytes);
 
diff --git a/Engine/Source/Programs/Unsync/Private/UnsyncMain.cpp b/Engine/Source/Programs/Unsync/Private/UnsyncMain.cpp
--- a/Engine/Source/Programs/Unsync/Private/UnsyncMain.cpp
+++ b/Engine/Source/Programs/Unsync/Private/UnsyncMain.cpp
@@ -11,6 +11,7 @@
 #include "UnsyncCmdQuery.h"
 #include "UnsyncCmdSync.h"
 #include "UnsyncCmdInfo.h"
+#include "UnsyncCmdVerify.h"
 #include "UnsyncCore.h"
 #include "UnsyncFile.h"
 #include "UnsyncMemory.h"
@@ -188,6 +189,12 @@ InnerMain(int Argc, char** Argv)
 		"Input files will be copied to compressed pack files during manifest generation and stored next to the manifest file.");
 	SubCommands.push_back(SubHash);
 
+	// Configure verify
+
+	CLI::App* SubVerify = Cli.add_subcommand("verify", "Verify directory contents against its manifest");
+	SubVerify->add_option("Input", InputFilenameUtf8, "Input directory path")->required();
+	SubCommands.push_back(SubVerify);
+
 	// Configure pack
 
 	CLI::App* SubPack = nullptr;
@@ -972,6 +979,13 @@ InnerMain(int Argc, char** Argv)
 
 		return CmdHash(HashOptions);
 	}
+	else if (Cli.got_subcommand(SubVerify))
+	{
+		FCmdVerifyOptions Options;
+		Options.Input = InputFilename;
+
+		return CmdVerify(Options);
+	}
 	else if (Cli.got_subcommand(SubPack))
 	{
 		FCmdPackOptions PackOptions;
diff --git a/Engine/Source/Programs/Unsync/Private/UnsyncVersion.cpp b/Engine/Source/Programs/Unsync/Private/UnsyncVersion.cpp
--- a/Engine/Source/Programs/Unsync/Private/UnsyncVersion.cpp
+++ b/Engine/Source/Programs/Unsync/Private/UnsyncVersion.cpp
@@ -2,7 +2,7 @@
 
 #include "UnsyncVersion.h"
 
-#define UNSYNC_VERSION_STR "1.0.94"
+#define UNSYNC_VERSION_STR "1.0.95"
 
 namespace unsync {
 
