unsync - Fix offset passed to async file read callback

* Aligned offset was being passed instead of original unaligned value

diff --git a/Engine/Source/Programs/Unsync/Private/UnsyncFile.cpp b/Engine/Source/Programs/Unsync/Private/UnsyncFile.cpp
--- a/Engine/Source/Programs/Unsync/Private/UnsyncFile.cpp
+++ b/Engine/Source/Programs/Unsync/Private/UnsyncFile.cpp
@@ -2,11 +2,11 @@
 
 #include "UnsyncFile.h"
 #include "UnsyncCore.h"
-#include "UnsyncMemory.h"
-#include "UnsyncThread.h"
-#include "UnsyncScheduler.h"
 #include "UnsyncFilter.h"
 #include "UnsyncHash.h"
+#include "UnsyncMemory.h"
+#include "UnsyncScheduler.h"
+#include "UnsyncThread.h"
 
 #include <mutex>
 
@@ -19,9 +19,9 @@
 
 #if UNSYNC_PLATFORM_WINDOWS
 UNSYNC_THIRD_PARTY_INCLUDES_START
-#include <winioctl.h>
+#	include <winioctl.h>
 UNSYNC_THIRD_PARTY_INCLUDES_END
-#endif // UNSYNC_PLATFORM_WINDOWS
+#endif	// UNSYNC_PLATFORM_WINDOWS
 
 namespace unsync {
 
@@ -59,9 +59,9 @@ MakeExtendedAbsolutePath(const FPath& InAbsolutePath)
 	{
 		return std::wstring(L"\\\\?\\") + InFilenameString;
 	}
-#else // UNSYNC_PLATFORM_WINDOWS
+#else	// UNSYNC_PLATFORM_WINDOWS
 	return InAbsolutePath;
-#endif // UNSYNC_PLATFORM_WINDOWS
+#endif	// UNSYNC_PLATFORM_WINDOWS
 }
 
 // Removes \\, \\?\UNC\, \\.\UNC\, \\.\ or \\?\ prefix from a path.
@@ -102,7 +102,7 @@ RemoveExtendedPathPrefix(const FPath& InPath)
 	if (InPathString.starts_with(L"\\\\?\\UNC\\"))
 	{
 		FPathStringView Remainder = InPathString.substr(8);
-		std::wstring Result;
+		std::wstring	Result;
 		Result.reserve(Remainder.length() + 2);
 		Result += L"\\\\";
 		Result += Remainder;
@@ -116,16 +116,17 @@ RemoveExtendedPathPrefix(const FPath& InPath)
 	{
 		return InPathString;
 	}
-#else // UNSYNC_PLATFORM_WINDOWS
+#else	// UNSYNC_PLATFORM_WINDOWS
 	return InPathString;
-#endif // UNSYNC_PLATFORM_WINDOWS
+#endif	// UNSYNC_PLATFORM_WINDOWS
 }
 
-std::filesystem::file_time_type FromWindowsFileTime(uint64 Ticks)
+std::filesystem::file_time_type
+FromWindowsFileTime(uint64 Ticks)
 {
 	using FileTimeDuration = std::filesystem::file_time_type::duration;
 
-	uint64 RawSeconds = Ticks / WINDOWS_TICKS_PER_SECOND;
+	uint64 RawSeconds		 = Ticks / WINDOWS_TICKS_PER_SECOND;
 	uint64 RawSubsecondTicks = Ticks - (RawSeconds * WINDOWS_TICKS_PER_SECOND);
 	uint64 RawSubsecondNanos = RawSubsecondTicks * NANOS_PER_WINDOWS_TICK;
 
@@ -178,41 +179,25 @@ GetRelativePathView(const FPath& Path, const FPath& Base)
 void
 ConvertDirectorySeparatorsToNative(std::string& Path)
 {
-	std::replace_if(
-		Path.begin(),
-		Path.end(),
-		[](char C) { return C == '/' || C == '\\'; },
-		PATH_SEPARATOR);
+	std::replace_if(Path.begin(), Path.end(), [](char C) { return C == '/' || C == '\\'; }, PATH_SEPARATOR);
 }
 
 void
 ConvertDirectorySeparatorsToUnix(std::string& Path)
 {
-	std::replace_if(
-		Path.begin(),
-		Path.end(),
-		[](char C) { return C == '\\'; },
-		char('/'));
+	std::replace_if(Path.begin(), Path.end(), [](char C) { return C == '\\'; }, char('/'));
 }
 
 void
 ConvertDirectorySeparatorsToNative(std::wstring& Path)
 {
-	std::replace_if(
-		Path.begin(),
-		Path.end(),
-		[](wchar_t C) { return C == '/' || C == '\\'; },
-		wchar_t(FPath::preferred_separator));
+	std::replace_if(Path.begin(), Path.end(), [](wchar_t C) { return C == '/' || C == '\\'; }, wchar_t(FPath::preferred_separator));
 }
 
 void
 ConvertDirectorySeparatorsToUnix(std::wstring& Path)
 {
-	std::replace_if(
-		Path.begin(),
-		Path.end(),
-		[](wchar_t C) { return C == '\\'; },
-		wchar_t('/'));
+	std::replace_if(Path.begin(), Path.end(), [](wchar_t C) { return C == '\\'; }, wchar_t('/'));
 }
 
 std::error_code
@@ -228,7 +213,6 @@ CopyFileIfNewer(const FPath& Source, const FPath& Target)
 	return Ec;
 }
 
-
 bool
 IsNonCaseSensitiveFileSystem(const FPath& ExistingPath)
 {
@@ -256,7 +240,8 @@ IsCaseSensitiveFileSystem(const FPath& ExistingPath)
 	return !IsNonCaseSensitiveFileSystem(ExistingPath);
 }
 
-FFileAttributes GetCachedFileAttrib(const FPath& Path, FFileAttributeCache& AttribCache)
+FFileAttributes
+GetCachedFileAttrib(const FPath& Path, FFileAttributeCache& AttribCache)
 {
 	FFileAttributes Result;
 
@@ -319,12 +304,14 @@ struct FWindowsAsyncFileReader final : FAsyncReader
 
 	struct FOverlappedCommand
 	{
-		OVERLAPPED Overlapped	= {};
-		uint64	   SourceOffset = 0;
-		uint64	   ReadSize		= 0;
-		uint64	   UserData		= 0;
-		uint32	   ErrorCode	= 0;
-		bool	   bActive		= false;
+		OVERLAPPED Overlapped	   = {};
+		uint64	   RequestedOffset = 0;
+		uint64	   RequestedSize   = 0;
+		uint64	   AlignedOffset   = 0;
+		uint64	   AlignedSize	   = 0;
+		uint64	   UserData		   = 0;
+		uint32	   ErrorCode	   = 0;
+		bool	   bActive		   = false;
 		FIOBuffer  Buffer;
 		IOCallback Callback = {};
 	};
@@ -348,9 +335,9 @@ struct FWindowsAsyncFileReader final : FAsyncReader
 	FOverlappedCommand Commands[MAX_OVERLAPPED_COMMANDS]		 = {};
 	uint64			   NumCommandsIssued						 = 0;
 
-	HANDLE		 FileHandle = INVALID_HANDLE_VALUE;
-	uint64		 FileSize	= 0;
-	FAtomicError ErrorCode;
+	HANDLE			 FileHandle = INVALID_HANDLE_VALUE;
+	uint64			 FileSize	= 0;
+	FAtomicError	 ErrorCode;
 	std::atomic_bool bClosed;
 };
 
@@ -454,12 +441,13 @@ FWindowsFile::FWindowsFile(const FPath& InFilename, EFileMode InMode, uint64 InS
 FWindowsFile::~FWindowsFile()
 {
 	Close();
-	
+
 	std::lock_guard<std::mutex> LockGuard(Mutex);
 	UNSYNC_ASSERT(AsyncReaders.empty());
 }
 
-void FWindowsFile::AddAsyncReader(FWindowsAsyncFileReader* Reader)
+void
+FWindowsFile::AddAsyncReader(FWindowsAsyncFileReader* Reader)
 {
 	std::lock_guard<std::mutex> LockGuard(Mutex);
 	AsyncReaders.push_back(Reader);
@@ -528,8 +516,8 @@ FWindowsAsyncFileReader::CompleteReadCommand(FOverlappedCommand& Cmd)
 
 	DWORD ReadBytes = 0;
 
-	const uint64 ExpectedReadEnd   = std::min(FileSize, Cmd.SourceOffset + Cmd.ReadSize);
-	const uint64 ExpectedReadBytes = ExpectedReadEnd - Cmd.SourceOffset;
+	const uint64 ExpectedReadEnd   = std::min(FileSize, Cmd.AlignedOffset + Cmd.AlignedSize);
+	const uint64 ExpectedReadBytes = ExpectedReadEnd - Cmd.AlignedOffset;
 
 	{
 		BOOL OverlappedResultOk = GetOverlappedResult(FileHandle, &Cmd.Overlapped, &ReadBytes, true);
@@ -546,20 +534,21 @@ FWindowsAsyncFileReader::CompleteReadCommand(FOverlappedCommand& Cmd)
 	}
 
 	// TODO: retry read command a few times to retry from errors
-	UNSYNC_ASSERTF(ReadBytes == ExpectedReadBytes,
-				   L"Expected to read %llu bytes, but read %llu [FileSize=%llu, Cmd.SourceOffset=%llu, Cmd.ReadSize=%llu, Cmd.ErrorCode=%u]",
-				   llu(ExpectedReadBytes),
-				   llu(ReadBytes),
-				   llu(FileSize),
-				   llu(Cmd.SourceOffset),
-				   llu(Cmd.ReadSize),
-				   Cmd.ErrorCode);
+	UNSYNC_ASSERTF(
+		ReadBytes == ExpectedReadBytes,
+		L"Expected to read %llu bytes, but read %llu [FileSize=%llu, Cmd.AlignedOffset=%llu, Cmd.AlignedSize=%llu, Cmd.ErrorCode=%u]",
+		llu(ExpectedReadBytes),
+		llu(ReadBytes),
+		llu(FileSize),
+		llu(Cmd.AlignedOffset),
+		llu(Cmd.AlignedSize),
+		Cmd.ErrorCode);
 
 	const uint64 ReadBytesClamped = std::min<uint64>(Cmd.Buffer.GetSize(), ReadBytes);
 
 	if (Cmd.Callback)
 	{
-		Cmd.Callback(std::move(Cmd.Buffer), Cmd.SourceOffset, ReadBytesClamped, Cmd.UserData);
+		Cmd.Callback(std::move(Cmd.Buffer), Cmd.RequestedOffset, ReadBytesClamped, Cmd.UserData);
 	}
 
 	Cmd.bActive = false;
@@ -578,7 +567,7 @@ FWindowsFile::Write(const void* Data, uint64 DestOffset, uint64 TotalSize)
 
 	if (!IsWriteOnly(Mode))
 	{
-		FlushAsyncReaders();	 // flush any outstanding read requests before writing
+		FlushAsyncReaders();  // flush any outstanding read requests before writing
 	}
 
 	LARGE_INTEGER Pos;
@@ -709,7 +698,7 @@ FWindowsAsyncFileReader::EnqueueRead(uint64 SourceOffset, uint64 Size, uint64 Us
 
 	// Async commands are always strictly ordered
 	const uint32 WaitSlotIndex = uint32(NumCommandsIssued % MaxQueueDepth);
-	const DWORD WaitResult = WaitForSingleObject(OverlappedEvents[WaitSlotIndex], INFINITE);
+	const DWORD	 WaitResult	   = WaitForSingleObject(OverlappedEvents[WaitSlotIndex], INFINITE);
 	UNSYNC_ASSERT(WaitResult == WAIT_OBJECT_0);
 	CmdIdx = WaitSlotIndex;
 
@@ -722,6 +711,9 @@ FWindowsAsyncFileReader::EnqueueRead(uint64 SourceOffset, uint64 Size, uint64 Us
 		CompleteReadCommand(Cmd);
 	}
 
+	Cmd.RequestedOffset = SourceOffset;
+	Cmd.RequestedSize	= Size;
+
 	if (EnumHasAnyFlags(FileMode, EFileMode::Unbuffered))
 	{
 		uint64 OriginalSize	 = Size;
@@ -746,11 +738,11 @@ FWindowsAsyncFileReader::EnqueueRead(uint64 SourceOffset, uint64 Size, uint64 Us
 	LARGE_INTEGER Pos;
 	Pos.QuadPart = SourceOffset;
 
-	Cmd.ReadSize			  = Size;
+	Cmd.AlignedSize			  = Size;
 	Cmd.Overlapped.Offset	  = Pos.LowPart;
 	Cmd.Overlapped.OffsetHigh = Pos.HighPart;
 	Cmd.UserData			  = UserData;
-	Cmd.SourceOffset		  = SourceOffset;
+	Cmd.AlignedOffset		  = SourceOffset;
 	Cmd.Callback			  = Callback;
 
 	ResetEvent(Cmd.Overlapped.hEvent);
@@ -788,7 +780,8 @@ FWindowsAsyncFileReader::Flush()
 	}
 }
 
-void FWindowsFile::FlushAsyncReaders()
+void
+FWindowsFile::FlushAsyncReaders()
 {
 	for (FWindowsAsyncFileReader* It : AsyncReaders)
 	{
@@ -947,7 +940,7 @@ GetFileAttrib(const FPath& Path, FFileAttributeCache* AttribCache)
 
 	// TODO: could potentially use std::filesystem::directory_entry for this on all platforms
 
-	std::error_code ErrorCode	  = {};
+	std::error_code ErrorCode = {};
 	auto			Entry	  = std::filesystem::directory_entry(Path, ErrorCode);
 
 	if (!ErrorCode)
@@ -960,12 +953,12 @@ GetFileAttrib(const FPath& Path, FFileAttributeCache* AttribCache)
 
 		std::filesystem::perms Perms = Status.permissions();
 
-		Result.bDirectory    = Entry.is_directory();
-		Result.Size		     = Result.bDirectory ? 0 : Entry.file_size();
-		Result.Mtime	     = ToWindowsFileTime(Entry.last_write_time());
-		Result.bReadOnly     = IsReadOnly(Perms);
+		Result.bDirectory	 = Entry.is_directory();
+		Result.Size			 = Result.bDirectory ? 0 : Entry.file_size();
+		Result.Mtime		 = ToWindowsFileTime(Entry.last_write_time());
+		Result.bReadOnly	 = IsReadOnly(Perms);
 		Result.bIsExecutable = IsExecutable(Perms);
-		Result.bValid	     = true;
+		Result.bValid		 = true;
 	}
 
 	return Result;
@@ -989,7 +982,7 @@ ToWindowsFileTime(const std::filesystem::file_time_type& FileTime)
 uint64
 GetAvailableDiskSpace(const FPath& Path)
 {
-	return ~0ull; // TODO: query available space via statvfs()
+	return ~0ull;  // TODO: query available space via statvfs()
 }
 
 #endif	// UNSYNC_PLATFORM_UNIX
@@ -1051,13 +1044,16 @@ SetFileExecutable(const FPath& Path, bool Executable)
 		std::filesystem::permissions(
 			ExtendedPath,
 			std::filesystem::perms::owner_exec | std::filesystem::perms::group_exec | std::filesystem::perms::others_exec,
-			std::filesystem::perm_options::add, ErrorCode);
+			std::filesystem::perm_options::add,
+			ErrorCode);
 	}
 	else
 	{
-		std::filesystem::permissions(ExtendedPath,
+		std::filesystem::permissions(
+			ExtendedPath,
 			std::filesystem::perms::owner_exec | std::filesystem::perms::group_exec | std::filesystem::perms::others_exec,
-			std::filesystem::perm_options::remove, ErrorCode);
+			std::filesystem::perm_options::remove,
+			ErrorCode);
 	}
 
 	return !ErrorCode;
@@ -1066,7 +1062,7 @@ SetFileExecutable(const FPath& Path, bool Executable)
 FBuffer
 ReadFileToBuffer(const FPath& Filename)
 {
-	FBuffer	   Result;
+	FBuffer		Result;
 	FNativeFile File(Filename, EFileMode::ReadOnly);
 	if (File.IsValid())
 	{
@@ -1107,7 +1103,8 @@ WriteBufferToFile(const FPath& Filename, const uint8* Data, uint64 Size, EFileMo
 		{
 			UNSYNC_ERROR(L"Failed to write complete file '%ls'. Expected to write %llu bytes, actually written %llu bytes",
 						 Filename.wstring().c_str(),
-						 llu(Size) , llu(WroteBytes));
+						 llu(Size),
+						 llu(WroteBytes));
 		}
 		return WroteBytes == Size;
 	}
@@ -1292,7 +1289,8 @@ CreateFileAttributeCache(const FPath& Root, const FSyncFilter* SyncFilter)
 
 	FTimePoint NextProgressLogTime = TimePointNow() + std::chrono::seconds(1);
 
-	auto ReportProgress = [&NextProgressLogTime, &Result]() {
+	auto ReportProgress = [&NextProgressLogTime, &Result]()
+	{
 		FTimePoint TimeNow = TimePointNow();
 		if (TimeNow >= NextProgressLogTime)
 		{
@@ -1315,13 +1313,13 @@ CreateFileAttributeCache(const FPath& Root, const FSyncFilter* SyncFilter)
 			continue;
 		}
 
-		FFileAttributes Attr = {};
+		FFileAttributes		   Attr	 = {};
 		std::filesystem::perms Perms = Dir.status().permissions();
 
-		Attr.Mtime	= ToWindowsFileTime(Dir.last_write_time());
-		Attr.Size	= Dir.file_size();
-		Attr.bValid = true;
-		Attr.bReadOnly = IsReadOnly(Perms);
+		Attr.Mtime		   = ToWindowsFileTime(Dir.last_write_time());
+		Attr.Size		   = Dir.file_size();
+		Attr.bValid		   = true;
+		Attr.bReadOnly	   = IsReadOnly(Perms);
 		Attr.bIsExecutable = IsExecutable(Perms);
 
 		Result.Map[Dir.path().native()] = Attr;
@@ -1400,7 +1398,8 @@ FileRemove(const FPath& Path, std::error_code& OutErrorCode)
 	return std::filesystem::remove(ExtendedPath, OutErrorCode);
 }
 
-std::filesystem::recursive_directory_iterator RecursiveDirectoryScan(const FPath& Path)
+std::filesystem::recursive_directory_iterator
+RecursiveDirectoryScan(const FPath& Path)
 {
 	FPath ExtendedPath = MakeExtendedAbsolutePath(Path);
 	return std::filesystem::recursive_directory_iterator(ExtendedPath);
@@ -1629,7 +1628,7 @@ TestFileAttrib()
 		SetFileReadOnly(TestFilename, false);
 	}
 
-	const bool	bFileWritten = WriteBufferToFile(TestFilename, "unsync test file");
+	const bool bFileWritten = WriteBufferToFile(TestFilename, "unsync test file");
 	UNSYNC_ASSERT(bFileWritten);
 
 	const uint64 ExpectedFileTime = 133425821060000000ull;
@@ -1668,7 +1667,7 @@ TestFileAttrib()
 	// means nothing there, so bIsExecutable will never be set to true.
 	const FFileAttributes FileAttribNotExecutable = GetFileAttrib(TestFilename);
 	UNSYNC_ASSERT(!FileAttribNotExecutable.bIsExecutable);
-#endif // UNSYNC_PLATFORM_UNIX
+#endif	// UNSYNC_PLATFORM_UNIX
 
 	std::error_code ErrorCode;
 	const bool		bFileDeleted = FileRemove(TestFilename, ErrorCode);
@@ -1760,7 +1759,7 @@ TestPathUtil()
 		UNSYNC_ASSERT(Relative.empty());
 	}
 
-#endif // UNSYNC_PLATFORM_WINDOWS
+#endif	// UNSYNC_PLATFORM_WINDOWS
 }
 
 void
@@ -1818,22 +1817,22 @@ DeleteOldFilesInDirectory(const FPath& Path, uint32 MaxFilesToKeep, bool bAllowI
 std::chrono::system_clock::time_point
 SystemTimeFromFileTime(std::filesystem::file_time_type FileTime)
 {
-#if 0 // TODO: use clock_cast() when it's supported in all target compilers
+#if 0  // TODO: use clock_cast() when it's supported in all target compilers
 
 	return std::chrono::clock_cast<std::chrono::system_clock>(FileTime);
 
 #else
 
-#if UNSYNC_PLATFORM_WINDOWS
+#	if UNSYNC_PLATFORM_WINDOWS
 	std::chrono::seconds Offset(SECONDS_BETWEEN_WINDOWS_AND_UNIX);
-#else // UNSYNC_PLATFORM_WINDOWS
+#	else	// UNSYNC_PLATFORM_WINDOWS
 	std::chrono::seconds Offset(0);
-#endif // UNSYNC_PLATFORM_WINDOWS
+#	endif	// UNSYNC_PLATFORM_WINDOWS
 
 	std::chrono::duration SinceEpoch = FileTime.time_since_epoch() - Offset;
 	return std::chrono::system_clock::time_point(std::chrono::duration_cast<std::chrono::system_clock::duration>(SinceEpoch));
 
-#endif // 0
+#endif	// 0
 }
 
 FDummyAsyncReader::FDummyAsyncReader(FIOReader& InReader) : Inner(InReader)
@@ -1938,8 +1937,8 @@ TestFileAsyncRead()
 
 		const FTimePoint HashDoneTime = TimePointNow();
 
-		const double ReadDuration = DurationSec(ReadStartTime, ReadDoneTime);
-		const double HashDuration = DurationSec(ReadDoneTime, HashDoneTime);
+		const double ReadDuration  = DurationSec(ReadStartTime, ReadDoneTime);
+		const double HashDuration  = DurationSec(ReadDoneTime, HashDoneTime);
 		const double TotalDuration = DurationSec(ReadStartTime, HashDoneTime);
 
 		UNSYNC_LOG(L"Read rate: %.2f MB / sec", SizeMb(TestFileSize) / ReadDuration);
