unsync - Handle possible partial overlapped file reads

diff --git a/Engine/Source/Programs/Unsync/Private/UnsyncFile.cpp b/Engine/Source/Programs/Unsync/Private/UnsyncFile.cpp
--- a/Engine/Source/Programs/Unsync/Private/UnsyncFile.cpp
+++ b/Engine/Source/Programs/Unsync/Private/UnsyncFile.cpp
@@ -309,9 +309,11 @@ struct FWindowsAsyncFileReader final : FAsyncReader
 		uint64	   RequestedSize   = 0;
 		uint64	   AlignedOffset   = 0;
 		uint64	   AlignedSize	   = 0;
+		uint64	   Transferred	   = 0;
 		uint64	   UserData		   = 0;
 		uint32	   ErrorCode	   = 0;
-		bool	   bActive		   = false;
+		bool	   bIoActive	   = false;
+		bool	   bComplete	   = true;
 		FIOBuffer  Buffer;
 		IOCallback Callback = {};
 	};
@@ -324,7 +326,10 @@ struct FWindowsAsyncFileReader final : FAsyncReader
 	virtual bool   EnqueueRead(uint64 SourceOffset, uint64 Size, uint64 UserData, IOCallback Callback);
 	virtual void   Flush() override;
 
-	uint32 CompleteReadCommand(FOverlappedCommand& Cmd);
+	bool BeginReadingNextSegment(FOverlappedCommand& Cmd);
+	bool FinishReadingSegment(FOverlappedCommand& Cmd);
+
+	void CompleteReadCommand(FOverlappedCommand& Cmd);
 
 	FWindowsFile& Inner;
 	const uint32  MaxQueueDepth;
@@ -509,51 +514,111 @@ FWindowsFile::Close()
 	}
 }
 
-uint32
-FWindowsAsyncFileReader::CompleteReadCommand(FOverlappedCommand& Cmd)
+inline void
+SetOverlappedOffset(OVERLAPPED& Overlapped, uint64 Offset)
 {
-	UNSYNC_ASSERT(Cmd.bActive);
+	LARGE_INTEGER Pos;
+	Pos.QuadPart = Offset;
+
+	Overlapped.Offset	  = Pos.LowPart;
+	Overlapped.OffsetHigh = Pos.HighPart;
+}
+
+bool
+FWindowsAsyncFileReader::FinishReadingSegment(FOverlappedCommand& Cmd)
+{
+	UNSYNC_ASSERT(Cmd.bIoActive);
 
 	DWORD ReadBytes = 0;
 
-	const uint64 ExpectedReadEnd   = std::min(FileSize, Cmd.AlignedOffset + Cmd.AlignedSize);
-	const uint64 ExpectedReadBytes = ExpectedReadEnd - Cmd.AlignedOffset;
+	const BOOL OverlappedResultOk = GetOverlappedResult(FileHandle, &Cmd.Overlapped, &ReadBytes, true);
+
+	Cmd.bIoActive = false;
+	Cmd.Transferred += ReadBytes;
+
+	if (OverlappedResultOk)
+	{
+		return true;
+	}
+	else
+	{
+		Cmd.ErrorCode = GetLastError();
+		ErrorCode.Set(SystemError(L"GetOverlappedResult failed", Cmd.ErrorCode));
+		return false;
+	}
+}
+
+bool
+FWindowsAsyncFileReader::BeginReadingNextSegment(FOverlappedCommand& Cmd)
+{
+	UNSYNC_ASSERT(!Cmd.bIoActive);
 
+	if (Cmd.Transferred >= Cmd.RequestedSize)
 	{
-		BOOL OverlappedResultOk = GetOverlappedResult(FileHandle, &Cmd.Overlapped, &ReadBytes, true);
+		return false;
+	}
+
+	if (EnumHasAnyFlags(Inner.Mode, EFileMode::Unbuffered))
+	{
+		Cmd.Transferred = AlignDownToMultiplePow2(Cmd.Transferred, FWindowsFile::UNBUFFERED_READ_ALIGNMENT);
+	}
 
-		if (OverlappedResultOk)
+	const uint64 NextReadSize = Cmd.AlignedSize - Cmd.Transferred;
+
+	uint8* BufferMemory = reinterpret_cast<uint8*>(Cmd.Buffer.GetMemory());
+	UNSYNC_ASSERT(Cmd.Transferred + NextReadSize <= Cmd.Buffer.GetMemorySize());
+
+	ResetEvent(Cmd.Overlapped.hEvent);
+
+	SetOverlappedOffset(Cmd.Overlapped, Cmd.AlignedOffset + Cmd.Transferred);
+
+	if (!ReadFile(FileHandle, BufferMemory + Cmd.Transferred, CheckedNarrow(NextReadSize), nullptr, &Cmd.Overlapped))
+	{
+		const DWORD LastError = GetLastError();
+		if (LastError != ERROR_IO_PENDING)
 		{
-			Cmd.ErrorCode = 0;
+			Cmd.ErrorCode = LastError;
+			ErrorCode.Set(SystemError(L"ReadFile failed", LastError));
+			return false;
 		}
-		else
+	}
+
+	Cmd.bIoActive = true;
+
+	return true;
+}
+
+void
+FWindowsAsyncFileReader::CompleteReadCommand(FOverlappedCommand& Cmd)
+{
+	UNSYNC_ASSERT(!Cmd.bComplete);
+
+	while (Cmd.bIoActive)
+	{
+		if (FinishReadingSegment(Cmd))
 		{
-			Cmd.ErrorCode = GetLastError();
-			ErrorCode.Set(SystemError(L"GetOverlappedResult failed", Cmd.ErrorCode));
+			BeginReadingNextSegment(Cmd);
 		}
 	}
 
-	// TODO: retry read command a few times to retry from errors
 	UNSYNC_ASSERTF(
-		ReadBytes == ExpectedReadBytes,
-		L"Expected to read %llu bytes, but read %llu [FileSize=%llu, Cmd.AlignedOffset=%llu, Cmd.AlignedSize=%llu, Cmd.ErrorCode=%u]",
-		llu(ExpectedReadBytes),
-		llu(ReadBytes),
+		Cmd.RequestedSize <= Cmd.Transferred,
+		L"Expected to read at least %llu bytes, but read %llu [FileSize=%llu, Cmd.AlignedOffset=%llu, Cmd.AlignedSize=%llu, Cmd.ErrorCode=%u]",
+		llu(Cmd.RequestedSize),
+		llu(Cmd.Transferred),
 		llu(FileSize),
 		llu(Cmd.AlignedOffset),
 		llu(Cmd.AlignedSize),
 		Cmd.ErrorCode);
 
-	const uint64 ReadBytesClamped = std::min<uint64>(Cmd.Buffer.GetSize(), ReadBytes);
+	const uint64 ReadBytesClamped = std::min<uint64>(Cmd.Buffer.GetSize(), Cmd.Transferred);
 
 	if (Cmd.Callback)
 	{
 		Cmd.Callback(std::move(Cmd.Buffer), Cmd.RequestedOffset, ReadBytesClamped, Cmd.UserData);
 	}
 
-	Cmd.bActive = false;
-
-	return CheckedNarrow(ReadBytesClamped);
+	Cmd.bComplete = true;
 }
 
 uint64
@@ -690,6 +755,9 @@ FWindowsAsyncFileReader::EnqueueRead(uint64 SourceOffset, uint64 Size, uint64 Us
 		return false;
 	}
 
+	// Clamp size to the end of the file
+	Size = std::min(FileSize, SourceOffset + Size) - SourceOffset;
+
 	const EFileMode FileMode = Inner.Mode;
 
 	UNSYNC_ASSERT(IsReadable(FileMode));
@@ -706,13 +774,18 @@ FWindowsAsyncFileReader::EnqueueRead(uint64 SourceOffset, uint64 Size, uint64 Us
 
 	FOverlappedCommand& Cmd = Commands[CmdIdx];
 
-	if (Cmd.bActive)
+	if (!Cmd.bComplete)
 	{
 		CompleteReadCommand(Cmd);
 	}
 
 	Cmd.RequestedOffset = SourceOffset;
 	Cmd.RequestedSize	= Size;
+	Cmd.UserData		= UserData;
+	Cmd.Callback		= Callback;
+	Cmd.Transferred		= 0;
+	Cmd.ErrorCode		= 0;
+	Cmd.bComplete		= false;
 
 	if (EnumHasAnyFlags(FileMode, EFileMode::Unbuffered))
 	{
@@ -723,48 +796,32 @@ FWindowsAsyncFileReader::EnqueueRead(uint64 SourceOffset, uint64 Size, uint64 Us
 		uint64 AlignedBegin = AlignDownToMultiplePow2(OriginalBegin, FWindowsFile::UNBUFFERED_READ_ALIGNMENT);
 		uint64 AlignedEnd	= AlignUpToMultiplePow2(OriginalEnd, FWindowsFile::UNBUFFERED_READ_ALIGNMENT);
 
-		SourceOffset = AlignedBegin;
-		Size		 = AlignedEnd - AlignedBegin;
+		uint64 AlignedSize = AlignedEnd - AlignedBegin;
 
-		Cmd.Buffer = FIOBuffer::Alloc(Size, L"WindowsFile::ReadAsync_aligned");
+		Cmd.Buffer = FIOBuffer::Alloc(AlignedSize, L"WindowsFile::ReadAsync_aligned");
 
 		Cmd.Buffer.SetDataRange(OriginalBegin - AlignedBegin, OriginalSize);
+
+		Cmd.AlignedOffset = AlignedBegin;
+		Cmd.AlignedSize	  = AlignedSize;
 	}
 	else
 	{
 		Cmd.Buffer = FIOBuffer::Alloc(Size, L"WindowsFile::ReadAsync");
-	}
-
-	LARGE_INTEGER Pos;
-	Pos.QuadPart = SourceOffset;
 
-	Cmd.AlignedSize			  = Size;
-	Cmd.Overlapped.Offset	  = Pos.LowPart;
-	Cmd.Overlapped.OffsetHigh = Pos.HighPart;
-	Cmd.UserData			  = UserData;
-	Cmd.AlignedOffset		  = SourceOffset;
-	Cmd.Callback			  = Callback;
-
-	ResetEvent(Cmd.Overlapped.hEvent);
+		Cmd.AlignedOffset = SourceOffset;
+		Cmd.AlignedSize	  = Size;
+	}
 
+	if (BeginReadingNextSegment(Cmd))
 	{
-		DWORD ReadOk = ReadFile(FileHandle, Cmd.Buffer.GetMemory(), CheckedNarrow(Size), nullptr, &Cmd.Overlapped);
-		if (!ReadOk)
-		{
-			const DWORD LastError = GetLastError();
-			if (LastError != ERROR_IO_PENDING)
-			{
-				Cmd.ErrorCode = LastError;
-				ErrorCode.Set(SystemError(L"ReadFile failed", LastError));
-			}
-		}
+		++NumCommandsIssued;
+		return true;
+	}
+	else
+	{
+		return false;
 	}
-
-	Cmd.bActive = true;
-
-	++NumCommandsIssued;
-
-	return true;
 }
 
 void
@@ -773,8 +830,9 @@ FWindowsAsyncFileReader::Flush()
 	for (uint32 I = 0; I < MaxQueueDepth; ++I)
 	{
 		const uint64 CommandIndex = (NumCommandsIssued + I) % MaxQueueDepth;
-		if (Commands[CommandIndex].bActive)
+		if (!Commands[CommandIndex].bComplete)
 		{
+			WaitForSingleObject(OverlappedEvents[CommandIndex], INFINITE);
 			CompleteReadCommand(Commands[CommandIndex]);
 		}
 	}
diff --git a/Engine/Source/Programs/Unsync/Private/UnsyncFile.h b/Engine/Source/Programs/Unsync/Private/UnsyncFile.h
--- a/Engine/Source/Programs/Unsync/Private/UnsyncFile.h
+++ b/Engine/Source/Programs/Unsync/Private/UnsyncFile.h
@@ -92,6 +92,12 @@ struct FIOBuffer
 		return DataSize;
 	}
 
+	uint64 GetMemorySize() const
+	{
+		UNSYNC_ASSERT(Canary == CANARY);
+		return MemorySize;
+	}
+
 	void Clear();
 
 	void SetDataRange(uint64 Offset, uint64 Size)
diff --git a/Engine/Source/Programs/Unsync/Private/UnsyncVersion.cpp b/Engine/Source/Programs/Unsync/Private/UnsyncVersion.cpp
--- a/Engine/Source/Programs/Unsync/Private/UnsyncVersion.cpp
+++ b/Engine/Source/Programs/Unsync/Private/UnsyncVersion.cpp
@@ -2,7 +2,7 @@
 
 #include "UnsyncVersion.h"
 
-#define UNSYNC_VERSION_STR "1.0.95"
+#define UNSYNC_VERSION_STR "1.0.96"
 
 namespace unsync {
 
