unsync - Catch and report all standard exceptions

diff --git a/Engine/Source/Programs/Unsync/Private/UnsyncMain.cpp b/Engine/Source/Programs/Unsync/Private/UnsyncMain.cpp
--- a/Engine/Source/Programs/Unsync/Private/UnsyncMain.cpp
+++ b/Engine/Source/Programs/Unsync/Private/UnsyncMain.cpp
@@ -1316,7 +1316,11 @@ main(int argc, char** argv)
 		}
 		catch (const std::system_error& E)
 		{
-			UNSYNC_ERROR(L"System error %d: %hs", E.code().value(), E.what())
+			UNSYNC_ERROR(L"System error %d: %hs", E.code().value(), E.what());
+		}
+		catch (const std::exception& E)
+		{
+			UNSYNC_ERROR(L"Unhandled exception: %hs", E.what());
 		}
 	}
 
diff --git a/Engine/Source/Programs/Unsync/Private/UnsyncVersion.cpp b/Engine/Source/Programs/Unsync/Private/UnsyncVersion.cpp
--- a/Engine/Source/Programs/Unsync/Private/UnsyncVersion.cpp
+++ b/Engine/Source/Programs/Unsync/Private/UnsyncVersion.cpp
@@ -2,7 +2,7 @@
 
 #include "UnsyncVersion.h"
 
-#define UNSYNC_VERSION_STR "1.0.96"
+#define UNSYNC_VERSION_STR "1.0.97"
 
 namespace unsync {
 
