unsync - Improve handling of socket connection errors during HTTP requests

Minor changes:
* Fix usage of deprecated function ZSTD_getDecompressedSize
* Remove spaces from CLI positional parameter names
* Bump version to 1.0.98

diff --git a/Engine/Source/Programs/Unsync/Private/UnsyncCmdQuery.cpp b/Engine/Source/Programs/Unsync/Private/UnsyncCmdQuery.cpp
--- a/Engine/Source/Programs/Unsync/Private/UnsyncCmdQuery.cpp
+++ b/Engine/Source/Programs/Unsync/Private/UnsyncCmdQuery.cpp
@@ -44,7 +44,10 @@ RunHttpPing(std::string_view Address, uint16 Port)
 	Request.Method			   = EHttpMethod::GET;
 
 	// Don't time the connection handshake, only query time itself
-	Connection.Open();
+	if (!Connection.Open())
+	{
+		return -1;
+	}
 
 	FTimePoint	  TimeBegin	   = TimePointNow();
 	FHttpResponse PingResponse = HttpRequest(Connection, Request);
diff --git a/Engine/Source/Programs/Unsync/Private/UnsyncCompression.cpp b/Engine/Source/Programs/Unsync/Private/UnsyncCompression.cpp
--- a/Engine/Source/Programs/Unsync/Private/UnsyncCompression.cpp
+++ b/Engine/Source/Programs/Unsync/Private/UnsyncCompression.cpp
@@ -54,7 +54,7 @@ Compress(const uint8* Data, uint64 DataSize, int ZstdCompressionLevel)
 bool
 Decompress(const uint8* InputData, uint64 InputDataSize, uint8* OutputData, uint64 OutputDataSize)
 {
-	uint64 DecompressedSizeExpected = ZSTD_getDecompressedSize(InputData, InputDataSize);
+	uint64 DecompressedSizeExpected = ZSTD_getFrameContentSize(InputData, InputDataSize);
 	if (OutputDataSize >= DecompressedSizeExpected)
 	{
 		uint64 DecompressedSizeActual = ZSTD_decompress(OutputData, OutputDataSize, InputData, InputDataSize);
@@ -67,7 +67,7 @@ Decompress(const uint8* InputData, uint64 InputDataSize, uint8* OutputData, uint
 FBuffer
 Decompress(const uint8* Data, uint64 DataSize)
 {
-	uint64	DecompressedSizeExpected = ZSTD_getDecompressedSize(Data, DataSize);
+	uint64	DecompressedSizeExpected = ZSTD_getFrameContentSize(Data, DataSize);
 	FBuffer Result(DecompressedSizeExpected);
 
 	if (DecompressedSizeExpected == 0)
diff --git a/Engine/Source/Programs/Unsync/Private/UnsyncHttp.cpp b/Engine/Source/Programs/Unsync/Private/UnsyncHttp.cpp
--- a/Engine/Source/Programs/Unsync/Private/UnsyncHttp.cpp
+++ b/Engine/Source/Programs/Unsync/Private/UnsyncHttp.cpp
@@ -590,6 +590,10 @@ FHttpConnection::Open()
 	if (TlsRequirement != ETlsRequirement::None)
 	{
 		FSocketHandle RawSocketHandle = SocketConnectTcp(HostAddress.c_str(), HostPort);
+		if (!SocketValid(RawSocketHandle))
+		{
+			return false;
+		}
 
 		FTlsClientSettings ClientSettings;
 		ClientSettings.bVerifyCertificate = bTlsVerifyCertificate;
@@ -617,6 +621,10 @@ FHttpConnection::Open()
 	if (!Socket && TlsRequirement != ETlsRequirement::Required)
 	{
 		FSocketHandle RawSocketHandle = SocketConnectTcp(HostAddress.c_str(), HostPort);
+		if (!SocketValid(RawSocketHandle))
+		{
+			return false;
+		}
 
 		Socket = std::unique_ptr<FSocketRaw>(new FSocketRaw(RawSocketHandle));
 	}
diff --git a/Engine/Source/Programs/Unsync/Private/UnsyncMain.cpp b/Engine/Source/Programs/Unsync/Private/UnsyncMain.cpp
--- a/Engine/Source/Programs/Unsync/Private/UnsyncMain.cpp
+++ b/Engine/Source/Programs/Unsync/Private/UnsyncMain.cpp
@@ -255,8 +255,8 @@ InnerMain(int Argc, char** Argv)
 	SubCommands.push_back(SubPush);
 
 	CLI::App* SubInfo = Cli.add_subcommand("info", "Display information about a manifest file or diff two manifests");
-	SubInfo->add_option("Input 1", InputFilenameUtf8, "Input manifest file or root directory")->required();
-	SubInfo->add_option("Input 2", InputFilename2Utf8, "Optional input manifest file or root directory");
+	SubInfo->add_option("InputA", InputFilenameUtf8, "Input manifest file or root directory")->required();
+	SubInfo->add_option("InputB", InputFilename2Utf8, "Optional input manifest file or root directory");
 	SubInfo->add_flag("--files", bInfoFiles, "List all files in the manifest");
 	SubInfo->add_option(
 		"--include",
diff --git a/Engine/Source/Programs/Unsync/Private/UnsyncSocket.cpp b/Engine/Source/Programs/Unsync/Private/UnsyncSocket.cpp
--- a/Engine/Source/Programs/Unsync/Private/UnsyncSocket.cpp
+++ b/Engine/Source/Programs/Unsync/Private/UnsyncSocket.cpp
@@ -105,7 +105,7 @@ ReportSocketError(int32 ErrorCode, ELogLevel LogLevel)
 	ErrorCode = WSAGetLastError();
 #endif
 
-	LogPrintf(LogLevel, L"Socket error: %d", ErrorCode);
+	LogPrintf(LogLevel, L"Socket error: %d\n", ErrorCode);
 	UNSYNC_BREAK_ON_ERROR_LEVEL(LogLevel);
 }
 
@@ -113,7 +113,7 @@ static void
 ReportSocketError(tls* TlsCtx, ELogLevel LogLevel)
 {
 	const char* ErrorString = tls_error(TlsCtx);
-	LogPrintf(LogLevel, L"TLS error: %hs", ErrorString);
+	LogPrintf(LogLevel, L"TLS error: %hs\n", ErrorString);
 	UNSYNC_BREAK_ON_ERROR_LEVEL(LogLevel);
 }
 
@@ -340,6 +340,11 @@ SocketSetRecvTimeout(FSocketHandle Socket, uint32 Seconds)
 bool
 SocketValid(FSocketHandle Socket)
 {
+	if (Socket == INVALID_SOCKET)
+	{
+		return false;
+	}
+
 	int RecvResult = send(Socket, nullptr, 0, 0);
 	return RecvResult != SOCKET_ERROR;
 }
diff --git a/Engine/Source/Programs/Unsync/Private/UnsyncVersion.cpp b/Engine/Source/Programs/Unsync/Private/UnsyncVersion.cpp
--- a/Engine/Source/Programs/Unsync/Private/UnsyncVersion.cpp
+++ b/Engine/Source/Programs/Unsync/Private/UnsyncVersion.cpp
@@ -2,7 +2,7 @@
 
 #include "UnsyncVersion.h"
 
-#define UNSYNC_VERSION_STR "1.0.97"
+#define UNSYNC_VERSION_STR "1.0.98"
 
 namespace unsync {
 
