Set receive timeouts on all Unsync connections.

diff --git a/Engine/Source/Programs/Unsync/Private/UnsyncAuth.cpp b/Engine/Source/Programs/Unsync/Private/UnsyncAuth.cpp
--- a/Engine/Source/Programs/Unsync/Private/UnsyncAuth.cpp
+++ b/Engine/Source/Programs/Unsync/Private/UnsyncAuth.cpp
@@ -279,7 +279,7 @@ DecodeJwtPayload(std::string JwtDataBase64Url)
 }
 
 TResult<FAuthToken>
-AcquireAuthToken(const FAuthDesc& AuthDesc, const FOpenIdConfig& OpenIdConfig)
+AcquireAuthToken(const FAuthDesc& AuthDesc, uint32 RecvTimeoutSeconds, const FOpenIdConfig& OpenIdConfig)
 {
 	if (OpenIdConfig.AuthorizationEndpoint.empty())
 	{
@@ -306,7 +306,7 @@ AcquireAuthToken(const FAuthDesc& AuthDesc, const FOpenIdConfig& OpenIdConfig)
 
 	FAuthToken Result;
 
-	TResult<FRemoteDesc> AuthRemoteDesc = FRemoteDesc::FromUrl(AuthDesc.AuthServer);
+	TResult<FRemoteDesc> AuthRemoteDesc = FRemoteDesc::FromUrl(AuthDesc.AuthServer, unsync::EProtocolFlavor::Unknown, RecvTimeoutSeconds);
 	if (AuthRemoteDesc.IsError())
 	{
 		return AppError(L"Failed to parse authentication server URI");
@@ -498,7 +498,7 @@ GetUserInfo(FHttpConnection& HttpConnection, const FAuthDesc& AuthDesc, const FO
 }
 
 TResult<FAuthToken>
-RefreshAuthToken(const FAuthDesc& AuthDesc, const FOpenIdConfig& OpenIdConfig, const FAuthToken& PreviousToken)
+RefreshAuthToken(const FAuthDesc& AuthDesc, uint32 RecvTimeoutSeconds, const FOpenIdConfig& OpenIdConfig, const FAuthToken& PreviousToken)
 {
 	if (OpenIdConfig.TokenEndpoint.empty())
 	{
@@ -507,7 +507,7 @@ RefreshAuthToken(const FAuthDesc& AuthDesc, const FOpenIdConfig& OpenIdConfig, c
 
 	FAuthToken Result = PreviousToken;
 
-	TResult<FRemoteDesc> AuthRemoteDesc = FRemoteDesc::FromUrl(AuthDesc.AuthServer);
+	TResult<FRemoteDesc> AuthRemoteDesc = FRemoteDesc::FromUrl(AuthDesc.AuthServer, unsync::EProtocolFlavor::Unknown, RecvTimeoutSeconds);
 	if (AuthRemoteDesc.IsError())
 	{
 		return AppError(L"Failed to parse authentication server URI");
@@ -876,12 +876,12 @@ LoadAuthToken(const FPath& Path)
 }
 
 TResult<FAuthToken>
-RefreshOrAcquireToken(const FAuthDesc& AuthDesc, const FOpenIdConfig& OpenIdConfig, const FAuthToken& PreviousToken)
+RefreshOrAcquireToken(const FAuthDesc& AuthDesc, uint32 RecvTimeoutSeconds, const FOpenIdConfig& OpenIdConfig, const FAuthToken& PreviousToken)
 {
 	if (!PreviousToken.Refresh.empty())
 	{
 		UNSYNC_VERBOSE(L"Refreshing access token");
-		TResult<FAuthToken> RefreshResult = RefreshAuthToken(AuthDesc, OpenIdConfig, PreviousToken);
+		TResult<FAuthToken> RefreshResult = RefreshAuthToken(AuthDesc, RecvTimeoutSeconds, OpenIdConfig, PreviousToken);
 		if (RefreshResult.IsOk())
 		{
 			return RefreshResult;
@@ -889,7 +889,7 @@ RefreshOrAcquireToken(const FAuthDesc& AuthDesc, const FOpenIdConfig& OpenIdConf
 	}
 
 	UNSYNC_VERBOSE(L"Requesting new access token");
-	return AcquireAuthToken(AuthDesc, OpenIdConfig);
+	return AcquireAuthToken(AuthDesc, RecvTimeoutSeconds, OpenIdConfig);
 }
 
 TResult<FPath>
@@ -950,7 +950,7 @@ FAuthDesc::FromHelloResponse(const ProxyQuery::FHelloResponse& HelloResponse)
 }
 
 TResult<FAuthToken>
-Authenticate(const FAuthDesc& AuthDesc, int32 RefreshThreshold)
+Authenticate(const FAuthDesc& AuthDesc, uint32 RecvTimeoutSeconds, int32 RefreshThreshold)
 {
 	// Authentication must be serialized (only one thread should ever open the browser for interactive login, etc.)
 	static std::mutex			AuthMutex;
@@ -995,13 +995,13 @@ Authenticate(const FAuthDesc& AuthDesc, int32 RefreshThreshold)
 		return ResultOk(PreviousToken);
 	}
 
-	TResult<FOpenIdConfig> OpenIdConfigResult = GetOpenIdConfig(AuthDesc);
+	TResult<FOpenIdConfig> OpenIdConfigResult = GetOpenIdConfig(AuthDesc, RecvTimeoutSeconds);
 	if (OpenIdConfigResult.IsError())
 	{
 		return MoveError<FAuthToken>(OpenIdConfigResult);
 	}
 
-	TResult<FAuthToken> FreshTokenResult = RefreshOrAcquireToken(AuthDesc, *OpenIdConfigResult, PreviousToken);
+	TResult<FAuthToken> FreshTokenResult = RefreshOrAcquireToken(AuthDesc, RecvTimeoutSeconds, * OpenIdConfigResult, PreviousToken);
 	if (FreshTokenResult.IsError())
 	{
 		return FreshTokenResult;
@@ -1044,14 +1044,14 @@ GetRemoteAuthDesc(const FRemoteDesc& RemoteDesc)
 }
 
 TResult<FOpenIdConfig>
-GetOpenIdConfig(const FAuthDesc& AuthDesc)
+GetOpenIdConfig(const FAuthDesc& AuthDesc, uint32 RecvTimeoutSeconds)
 {
 	if (!AuthDesc.IsValid())
 	{
 		return AppError(L"Mandatory authentication parameters not provided");
 	}
 
-	TResult<FRemoteDesc> AuthServerDescResult = FRemoteDesc::FromUrl(AuthDesc.AuthServer);
+	TResult<FRemoteDesc> AuthServerDescResult = FRemoteDesc::FromUrl(AuthDesc.AuthServer, unsync::EProtocolFlavor::Unknown, RecvTimeoutSeconds);
 	if (AuthServerDescResult.IsError())
 	{
 		return MoveError<FOpenIdConfig>(AuthServerDescResult);
diff --git a/Engine/Source/Programs/Unsync/Private/UnsyncAuth.h b/Engine/Source/Programs/Unsync/Private/UnsyncAuth.h
--- a/Engine/Source/Programs/Unsync/Private/UnsyncAuth.h
+++ b/Engine/Source/Programs/Unsync/Private/UnsyncAuth.h
@@ -61,20 +61,20 @@ struct FAuthDesc
 };
 
 // Query OIDC/OAuth2 configuration via HTTP
-TResult<FOpenIdConfig> GetOpenIdConfig(const FAuthDesc& AuthDesc);
+TResult<FOpenIdConfig> GetOpenIdConfig(const FAuthDesc& AuthDesc, uint32 RecvTimeoutSeconds);
 
 // Perform the complete authentication flow:
 // - Attempt use stored refresh token first
 // - If refresh is not possible, use PKCE Authentication flow to get new tokens
 // - Save refresh token in user directory for future use
 // - Skips acquiring new token if remaining valid time is above RefreshThreshold (in seconds, 5 minutes by default)
-TResult<FAuthToken> Authenticate(const FAuthDesc& AuthDesc, int32 RefreshThreshold = 300);
+TResult<FAuthToken> Authenticate(const FAuthDesc& AuthDesc, uint32 RecvTimeoutSeconds, int32 RefreshThreshold = 300);
 
 // Auth utility functions
 
-TResult<FAuthToken> AcquireAuthToken(const FAuthDesc& AuthDesc, const FOpenIdConfig& OpenIdConfig);
-TResult<FAuthToken> RefreshAuthToken(const FAuthDesc& AuthDesc, const FOpenIdConfig& OpenIdConfig, const FAuthToken& PreviousToken);
-TResult<FAuthToken> RefreshOrAcquireToken(const FAuthDesc& AuthDesc, const FOpenIdConfig& OpenIdConfig, const FAuthToken& PreviousToken);
+TResult<FAuthToken> AcquireAuthToken(const FAuthDesc& AuthDesc, uint32 RecvTimeoutSeconds, const FOpenIdConfig& OpenIdConfig);
+TResult<FAuthToken> RefreshAuthToken(const FAuthDesc& AuthDesc, uint32 RecvTimeoutSeconds, const FOpenIdConfig& OpenIdConfig, const FAuthToken& PreviousToken);
+TResult<FAuthToken> RefreshOrAcquireToken(const FAuthDesc& AuthDesc, uint32 RecvTimeoutSeconds, const FOpenIdConfig& OpenIdConfig, const FAuthToken& PreviousToken);
 
 struct FAuthUserInfo
 {
@@ -85,9 +85,9 @@ struct FAuthUserInfo
 	std::string FamilyName;
 	std::string Email;
 };
-TResult<FAuthUserInfo> GetUserInfo(FHttpConnection& HttpConnection, const FAuthDesc& AuthDesc, const FAuthToken& AuthToken);
+TResult<FAuthUserInfo> GetUserInfo(FHttpConnection& HttpConnection, const FAuthDesc& AuthDesc, const FOpenIdConfig& OpenIdConfig, const FAuthToken& AuthToken);
 
-std::string GenerateTokenId(const FRemoteDesc& RemoteDesc);
+std::string GenerateTokenId(const FAuthDesc& RemoteDesc);
 
 bool				SaveAuthToken(const FPath& Path, const FAuthToken& AuthToken);
 TResult<FAuthToken> LoadAuthToken(const FPath& Path);
diff --git a/Engine/Source/Programs/Unsync/Private/UnsyncCmdLogin.cpp b/Engine/Source/Programs/Unsync/Private/UnsyncCmdLogin.cpp
--- a/Engine/Source/Programs/Unsync/Private/UnsyncCmdLogin.cpp
+++ b/Engine/Source/Programs/Unsync/Private/UnsyncCmdLogin.cpp
@@ -46,7 +46,7 @@ CmdLogin(const FCmdLoginOptions& Options)
 
 	const FAuthDesc AuthDesc = FAuthDesc::FromHelloResponse(*HelloResponseResult);
 
-	TResult<FAuthToken> AuthTokenResult = Authenticate(AuthDesc, RefreshThreshold);
+	TResult<FAuthToken> AuthTokenResult = Authenticate(AuthDesc, Options.Remote.RecvTimeoutSeconds, RefreshThreshold);
 
 	if (AuthTokenResult.IsOk())
 	{
diff --git a/Engine/Source/Programs/Unsync/Private/UnsyncCmdQuery.cpp b/Engine/Source/Programs/Unsync/Private/UnsyncCmdQuery.cpp
--- a/Engine/Source/Programs/Unsync/Private/UnsyncCmdQuery.cpp
+++ b/Engine/Source/Programs/Unsync/Private/UnsyncCmdQuery.cpp
@@ -37,7 +37,9 @@ RunHttpPing(std::string_view Address, uint16 Port)
 		TlsRequirement = ETlsRequirement::Preferred;
 	}
 
-	FHttpConnection Connection(Address, Port, TlsRequirement, TlsSettings);
+	constexpr uint32 RecvTimeoutSeconds = 4;
+
+	FHttpConnection Connection(Address, Port, TlsRequirement, TlsSettings, RecvTimeoutSeconds);
 
 	FHttpRequest	Request;
 	Request.Url				   = "/api/v1/ping";
@@ -217,7 +219,7 @@ CmdQuerySearch(const FCmdQueryOptions& Options)
 	auto CreateConnection = [Remote = Options.Remote]
 	{
 		FTlsClientSettings TlsSettings = Remote.GetTlsClientSettings();
-		return new FHttpConnection(Remote.Host.Address, Remote.Host.Port, Remote.TlsRequirement, TlsSettings);
+		return new FHttpConnection(Remote.Host.Address, Remote.Host.Port, Remote.TlsRequirement, TlsSettings, Remote.RecvTimeoutSeconds);
 	};
 
 	TObjectPool<FHttpConnection> ConnectionPool(CreateConnection);
@@ -450,7 +452,7 @@ CmdQueryFile(const FCmdQueryOptions& Options)
 	FAuthDesc AuthDesc = FAuthDesc::FromHelloResponse(*HelloResponse);
 	HelloConnection.Close();
 
-	TResult<FAuthToken> AuthToken = Authenticate(AuthDesc);
+	TResult<FAuthToken> AuthToken = Authenticate(AuthDesc, Options.Remote.RecvTimeoutSeconds);
 
 	if (!AuthToken.IsOk())
 	{
@@ -522,7 +524,7 @@ CmdQueryHttpGet(const FCmdQueryOptions& Options)
 		}
 
 		FAuthDesc			AuthDesc  = FAuthDesc::FromHelloResponse(*HelloResponse);
-		TResult<FAuthToken> AuthToken = Authenticate(AuthDesc);
+		TResult<FAuthToken> AuthToken = Authenticate(AuthDesc, Options.Remote.RecvTimeoutSeconds);
 
 		if (!AuthToken.IsOk())
 		{
diff --git a/Engine/Source/Programs/Unsync/Private/UnsyncHorde.cpp b/Engine/Source/Programs/Unsync/Private/UnsyncHorde.cpp
--- a/Engine/Source/Programs/Unsync/Private/UnsyncHorde.cpp
+++ b/Engine/Source/Programs/Unsync/Private/UnsyncHorde.cpp
@@ -586,7 +586,7 @@ FHordeProtocolImpl::QueryListDirectory(FHttpConnection& Connection, const FAuthD
 
 	std::string BearerToken;
 
-	TResult<FAuthToken> AuthTokenResult = Authenticate(*AuthDesc);
+	TResult<FAuthToken> AuthTokenResult = Authenticate(*AuthDesc, Connection.RecvTimeoutSeconds);
 	if (const FAuthToken* AuthToken = AuthTokenResult.TryData())
 	{
 		BearerToken = AuthToken->Access;
diff --git a/Engine/Source/Programs/Unsync/Private/UnsyncHttp.cpp b/Engine/Source/Programs/Unsync/Private/UnsyncHttp.cpp
--- a/Engine/Source/Programs/Unsync/Private/UnsyncHttp.cpp
+++ b/Engine/Source/Programs/Unsync/Private/UnsyncHttp.cpp
@@ -286,7 +286,7 @@ HttpRequest(const FRemoteDesc& RemoteDesc,
 			std::string_view   BearerToken)
 {
 	FTlsClientSettings TlsSettings = RemoteDesc.GetTlsClientSettings();
-	FHttpConnection	   Connection(RemoteDesc.Host.Address, RemoteDesc.Host.Port, RemoteDesc.TlsRequirement, TlsSettings);
+	FHttpConnection	   Connection(RemoteDesc.Host.Address, RemoteDesc.Host.Port, RemoteDesc.TlsRequirement, TlsSettings, RemoteDesc.RecvTimeoutSeconds);
 
 	FHttpRequest Request;
 
@@ -520,10 +520,12 @@ HttpRequestEnd(FHttpConnection& Connection, FHttpChunkCallback ChunkCallback)
 FHttpConnection::FHttpConnection(const std::string_view	   InHostAddress,
 								 uint16					   InPort,
 								 ETlsRequirement		   InTlsRequirement,
-								 const FTlsClientSettings& InTlsSettings)
+								 const FTlsClientSettings& InTlsSettings,
+								 uint32					   InRecvTimeoutSeconds)
 : HostAddress(InHostAddress)
 , HostPort(InPort)
 , TlsRequirement(InTlsRequirement)
+, RecvTimeoutSeconds(InRecvTimeoutSeconds)
 {
 	if (InTlsSettings.Subject.empty())
 	{
@@ -552,28 +554,29 @@ FHttpConnection::FHttpConnection(const FHttpConnection& Other)
 , TlsSubject(Other.TlsSubject)
 , bTlsVerifyCertificate(Other.bTlsVerifyCertificate)
 , TlsCacert(Other.TlsCacert)
+, RecvTimeoutSeconds(Other.RecvTimeoutSeconds)
 {
 }
 
 FHttpConnection
-FHttpConnection::CreateDefaultHttp(const std::string_view InHostAddress, uint16 Port)
+FHttpConnection::CreateDefaultHttp(const std::string_view InHostAddress, uint16 Port, uint32 RecvTimeoutSeconds)
 {
-	return FHttpConnection(InHostAddress, Port);
+	return FHttpConnection(InHostAddress, Port, RecvTimeoutSeconds);
 }
 
 FHttpConnection
-FHttpConnection::CreateDefaultHttps(const std::string_view InHostAddress, uint16 Port)
+FHttpConnection::CreateDefaultHttps(const std::string_view InHostAddress, uint16 Port, uint32 RecvTimeoutSeconds)
 {
 	FTlsClientSettings TlsSettings;
 	TlsSettings.Subject = InHostAddress.data();
-	return FHttpConnection(InHostAddress, Port, ETlsRequirement::Required, TlsSettings);
+	return FHttpConnection(InHostAddress, Port, ETlsRequirement::Required, TlsSettings, RecvTimeoutSeconds);
 }
 
 FHttpConnection
 FHttpConnection::CreateDefaultHttps(const FRemoteDesc& RemoteDesc)
 {
 	FTlsClientSettings TlsSettings = RemoteDesc.GetTlsClientSettings();
-	return FHttpConnection(RemoteDesc.Host.Address, RemoteDesc.Host.Port, ETlsRequirement::Required, TlsSettings);
+	return FHttpConnection(RemoteDesc.Host.Address, RemoteDesc.Host.Port, ETlsRequirement::Required, TlsSettings, RemoteDesc.RecvTimeoutSeconds);
 }
 
 bool
@@ -595,6 +598,8 @@ FHttpConnection::Open()
 			return false;
 		}
 
+		SocketSetRecvTimeout(RawSocketHandle, RecvTimeoutSeconds);
+
 		FTlsClientSettings ClientSettings;
 		ClientSettings.bVerifyCertificate = bTlsVerifyCertificate;
 		ClientSettings.bVerifySubject	  = bTlsVerifySubject;
@@ -626,6 +631,8 @@ FHttpConnection::Open()
 			return false;
 		}
 
+		SocketSetRecvTimeout(RawSocketHandle, RecvTimeoutSeconds);
+
 		Socket = std::unique_ptr<FSocketRaw>(new FSocketRaw(RawSocketHandle));
 	}
 
diff --git a/Engine/Source/Programs/Unsync/Private/UnsyncHttp.h b/Engine/Source/Programs/Unsync/Private/UnsyncHttp.h
--- a/Engine/Source/Programs/Unsync/Private/UnsyncHttp.h
+++ b/Engine/Source/Programs/Unsync/Private/UnsyncHttp.h
@@ -73,12 +73,12 @@ struct FHttpResponse
 
 struct FHttpConnection
 {
-	FHttpConnection(const std::string_view InHostAddress, uint16 InPort) : FHttpConnection(InHostAddress, InPort, ETlsRequirement::None, FTlsClientSettings{}) {}
-	FHttpConnection(const std::string_view InHostAddress, uint16 InPort, ETlsRequirement InTlsRequirement, const FTlsClientSettings& InTlsSettings);
+	FHttpConnection(const std::string_view InHostAddress, uint16 InPort, uint32 InRecvTimeoutSeconds) : FHttpConnection(InHostAddress, InPort, ETlsRequirement::None, FTlsClientSettings{}, InRecvTimeoutSeconds) {}
+	FHttpConnection(const std::string_view InHostAddress, uint16 InPort, ETlsRequirement InTlsRequirement, const FTlsClientSettings& InTlsSettings, uint32 InRecvTimeoutSeconds);
 	FHttpConnection(const FHttpConnection& Other);
 
-	[[nodiscard]] static FHttpConnection CreateDefaultHttp(const std::string_view InHostAddress, uint16 Port = 80);
-	[[nodiscard]] static FHttpConnection CreateDefaultHttps(const std::string_view InHostAddress, uint16 Port = 443);
+	[[nodiscard]] static FHttpConnection CreateDefaultHttp(const std::string_view InHostAddress, uint16 Port = 80, uint32 RecvTimeoutSeconds = 0);
+	[[nodiscard]] static FHttpConnection CreateDefaultHttps(const std::string_view InHostAddress, uint16 Port = 443, uint32 RecvTimeoutSeconds = 0);
 	[[nodiscard]] static FHttpConnection CreateDefaultHttps(const FRemoteDesc& RemoteDesc);
 
 	bool Open();
@@ -95,6 +95,8 @@ struct FHttpConnection
 	bool					 bTlsVerifyCertificate = true;
 	std::shared_ptr<FBuffer> TlsCacert;
 
+	uint32 RecvTimeoutSeconds = 0;
+
 	uint64 NumActiveRequests = 0;
 
 	// TODO: use single memory allocation for multiple reponse objects (perhaps a ring buffer)
diff --git a/Engine/Source/Programs/Unsync/Private/UnsyncJupiter.cpp b/Engine/Source/Programs/Unsync/Private/UnsyncJupiter.cpp
--- a/Engine/Source/Programs/Unsync/Private/UnsyncJupiter.cpp
+++ b/Engine/Source/Programs/Unsync/Private/UnsyncJupiter.cpp
@@ -459,7 +459,7 @@ TResult<uint64>
 JupiterPush(const FDirectoryManifest& Manifest, const FRemoteDesc& RemoteDesc, const FTlsClientSettings& TlsSettings)
 {
 	auto CreateConnection = [RemoteDesc, TlsSettings]
-	{ return new FHttpConnection(RemoteDesc.Host.Address, RemoteDesc.Host.Port, RemoteDesc.TlsRequirement, TlsSettings); };
+	{ return new FHttpConnection(RemoteDesc.Host.Address, RemoteDesc.Host.Port, RemoteDesc.TlsRequirement, TlsSettings, RemoteDesc.RecvTimeoutSeconds); };
 
 	TObjectPool<FHttpConnection> ConnectionPool(CreateConnection);
 
@@ -692,7 +692,7 @@ FJupiterProtocolImpl::FJupiterProtocolImpl(const FRemoteDesc&		 InSettings,
 										   const FBlockRequestMap*	 InRequestMap,
 										   std::string_view			 InHttpHeaders)
 : FRemoteProtocolBase(InSettings, InRequestMap)
-, Connection(InSettings.Host.Address, InSettings.Host.Port, InSettings.TlsRequirement, InSettings.GetTlsClientSettings())
+, Connection(InSettings.Host.Address, InSettings.Host.Port, InSettings.TlsRequirement, InSettings.GetTlsClientSettings(), InSettings.RecvTimeoutSeconds)
 , HttpHeaders(InHttpHeaders)
 , RemoteDesc(InSettings)
 {
diff --git a/Engine/Source/Programs/Unsync/Private/UnsyncMain.cpp b/Engine/Source/Programs/Unsync/Private/UnsyncMain.cpp
--- a/Engine/Source/Programs/Unsync/Private/UnsyncMain.cpp
+++ b/Engine/Source/Programs/Unsync/Private/UnsyncMain.cpp
@@ -928,7 +928,14 @@ InnerMain(int Argc, char** Argv)
 		RemoteDesc.bAuthenticationRequired = true;
 	}
 
-	FRemoteDesc RootRemoteDesc = RemoteDesc;
+	if (bNoSocketTimeout)
+	{
+		RemoteDesc.RecvTimeoutSeconds = 0;
+	}
+	else
+	{
+		RemoteDesc.RecvTimeoutSeconds = 15 * 60;
+	}
 
 	if (!bNoProxySelect
 		&& Cli.got_subcommand(SubSync)
@@ -1050,7 +1057,7 @@ InnerMain(int Argc, char** Argv)
 
 				// Note: since tokens can expire during a long operation,
 				// we can only save the auth descriptor and re-authenticate later if necessary
-				TResult<FAuthToken> AuthTokenResult = Authenticate(AuthDesc);
+				TResult<FAuthToken> AuthTokenResult = Authenticate(AuthDesc, RemoteDesc.RecvTimeoutSeconds);
 
 				if (AuthTokenResult.IsError())
 				{
@@ -1067,15 +1074,6 @@ InnerMain(int Argc, char** Argv)
 			}
 		}
 
-		if (bNoSocketTimeout)
-		{
-			RemoteDesc.RecvTimeoutSeconds = 0;
-		}
-		else
-		{
-			RemoteDesc.RecvTimeoutSeconds = 15 * 60;
-		}
-
 		// Try to derive default memory budget
 
 		if (BackgroundMemoryBudgetOption->empty())
diff --git a/Engine/Source/Programs/Unsync/Private/UnsyncProxy.cpp b/Engine/Source/Programs/Unsync/Private/UnsyncProxy.cpp
--- a/Engine/Source/Programs/Unsync/Private/UnsyncProxy.cpp
+++ b/Engine/Source/Programs/Unsync/Private/UnsyncProxy.cpp
@@ -192,7 +192,7 @@ FUnsyncProtocolImpl::FUnsyncProtocolImpl(const FRemoteDesc&				RemoteDesc,
 		{
 			bool bOk = IsValid();
 
-			TResult<FAuthToken> AuthTokenResult = Authenticate(*InAuthDesc);
+			TResult<FAuthToken> AuthTokenResult = Authenticate(*InAuthDesc, RemoteDesc.RecvTimeoutSeconds);
 
 			if (AuthTokenResult.IsOk())
 			{
@@ -682,7 +682,7 @@ TResult<ProxyQuery::FHelloResponse>
 ProxyQuery::Hello(const FRemoteDesc& RemoteDesc, const FAuthDesc* OptAuthDesc)
 {
 	FTlsClientSettings TlsSettings = RemoteDesc.GetTlsClientSettings();
-	FHttpConnection	   Connection(RemoteDesc.Host.Address, RemoteDesc.Host.Port, RemoteDesc.TlsRequirement, TlsSettings);
+	FHttpConnection	   Connection(RemoteDesc.Host.Address, RemoteDesc.Host.Port, RemoteDesc.TlsRequirement, TlsSettings, RemoteDesc.RecvTimeoutSeconds);
 
 	return Hello(RemoteDesc.Protocol, Connection, OptAuthDesc);
 }
@@ -714,7 +714,7 @@ FUnsyncProtocolImpl::QueryHello(FHttpConnection& HttpConnection, const FAuthDesc
 	std::string BearerToken;
 	if (OptAuthDesc)
 	{
-		TResult<FAuthToken> AuthTokenResult = Authenticate(*OptAuthDesc);
+		TResult<FAuthToken> AuthTokenResult = Authenticate(*OptAuthDesc, HttpConnection.RecvTimeoutSeconds);
 		if (AuthTokenResult.IsOk())
 		{
 			BearerToken = std::move(AuthTokenResult.GetData().Access);
@@ -938,7 +938,7 @@ FUnsyncProtocolImpl::QueryListDirectory(FHttpConnection& Connection, const FAuth
 	std::string BearerToken;
 	if (AuthDesc)
 	{
-		TResult<FAuthToken> AuthToken = Authenticate(*AuthDesc);
+		TResult<FAuthToken> AuthToken = Authenticate(*AuthDesc, Connection.RecvTimeoutSeconds);
 		if (!AuthToken.IsOk())
 		{
 			return MoveError<FDirectoryListing>(AuthToken);
@@ -987,7 +987,7 @@ ProxyQuery::DownloadFile(FHttpConnection&					 InConnection,
 		std::string BearerToken;
 		if (AuthDesc)
 		{
-			TResult<FAuthToken> AuthToken = Authenticate(*AuthDesc);
+			TResult<FAuthToken> AuthToken = Authenticate(*AuthDesc, Connection->RecvTimeoutSeconds);
 			if (!AuthToken.IsOk())
 			{
 				return std::move(AuthToken.GetError());
@@ -1038,7 +1038,7 @@ ProxyQuery::DownloadFile(FHttpConnection&					 InConnection,
 	std::string BearerToken;
 	if (AuthDesc)
 	{
-		TResult<FAuthToken> AuthToken = Authenticate(*AuthDesc);
+		TResult<FAuthToken> AuthToken = Authenticate(*AuthDesc, InConnection.RecvTimeoutSeconds);
 		if (!AuthToken.IsOk())
 		{
 			Error.Set(std::move(AuthToken.GetError()));
@@ -1334,7 +1334,7 @@ FProxyPool::FProxyPool(const FRemoteDesc& InRemoteDesc, const FAuthDesc* InAuthD
 	auto CreateHttpConnection = [Remote = RemoteDesc]
 	{
 		FTlsClientSettings TlsSettings = Remote.GetTlsClientSettings();
-		return new FHttpConnection(Remote.Host.Address, Remote.Host.Port, Remote.TlsRequirement, TlsSettings);
+		return new FHttpConnection(Remote.Host.Address, Remote.Host.Port, Remote.TlsRequirement, TlsSettings, Remote.RecvTimeoutSeconds);
 	};
 
 	HttpPool.emplace(CreateHttpConnection);
@@ -1445,7 +1445,7 @@ FProxyPool::GetAccessToken()
 
 	if (AuthDesc)
 	{
-		TResult<FAuthToken> AuthTokenResult = Authenticate(*AuthDesc);
+		TResult<FAuthToken> AuthTokenResult = Authenticate(*AuthDesc, RemoteDesc.RecvTimeoutSeconds);
 		if (FAuthToken* AuthToken = AuthTokenResult.TryData())
 		{
 			std::swap(Result, AuthToken->Access);
diff --git a/Engine/Source/Programs/Unsync/Private/UnsyncRemote.cpp b/Engine/Source/Programs/Unsync/Private/UnsyncRemote.cpp
--- a/Engine/Source/Programs/Unsync/Private/UnsyncRemote.cpp
+++ b/Engine/Source/Programs/Unsync/Private/UnsyncRemote.cpp
@@ -75,7 +75,7 @@ IsValidUrl(std::string_view Url)
 }
 
 TResult<FRemoteDesc>
-FRemoteDesc::FromUrl(std::string_view Url, EProtocolFlavor ProtocolFlavorHint)
+FRemoteDesc::FromUrl(std::string_view Url, EProtocolFlavor ProtocolFlavorHint, uint32 RecvTimeoutSeconds)
 {
 	if (!IsValidUrl(Url))
 	{
@@ -208,6 +208,7 @@ FRemoteDesc::FromUrl(std::string_view Url, EProtocolFlavor ProtocolFlavorHint)
 	Result.Host.Address	  = HostAddress;
 	Result.Host.Port	  = HostPort;
 	Result.TlsRequirement = TlsRequirement;
+	Result.RecvTimeoutSeconds = RecvTimeoutSeconds;
 
 	return ResultOk(Result);
 }
diff --git a/Engine/Source/Programs/Unsync/Private/UnsyncRemote.h b/Engine/Source/Programs/Unsync/Private/UnsyncRemote.h
--- a/Engine/Source/Programs/Unsync/Private/UnsyncRemote.h
+++ b/Engine/Source/Programs/Unsync/Private/UnsyncRemote.h
@@ -72,7 +72,7 @@ struct FRemoteDesc
 
 	bool IsValid() const { return Protocol != EProtocolFlavor::Unknown && Host.IsValid(); }
 
-	static TResult<FRemoteDesc> FromUrl(std::string_view Url, EProtocolFlavor ProtocolFlavorHint = EProtocolFlavor::Unknown);
+	static TResult<FRemoteDesc> FromUrl(std::string_view Url, EProtocolFlavor ProtocolFlavorHint = EProtocolFlavor::Unknown, uint32 RecvTimeoutSeconds = 0);
 
 	FTlsClientSettings GetTlsClientSettings() const;
 };
diff --git a/Engine/Source/Programs/Unsync/Private/UnsyncVersion.cpp b/Engine/Source/Programs/Unsync/Private/UnsyncVersion.cpp
--- a/Engine/Source/Programs/Unsync/Private/UnsyncVersion.cpp
+++ b/Engine/Source/Programs/Unsync/Private/UnsyncVersion.cpp
@@ -2,7 +2,7 @@
 
 #include "UnsyncVersion.h"
 
-#define UNSYNC_VERSION_STR "1.0.98"
+#define UNSYNC_VERSION_STR "1.0.99"
 
 namespace unsync {
 
