Unsync: Add block cloning (reflinking) to minimize disk I/O

Bytes reused without copying are logged as "in situ", and displayed with --very-verbose.

#tests Unit tests on ReFS, btrfs & tmpfs. Real cooked builds on ReFS.
[RN] unsync: Added copy-on-write block cloning (reflinking) to reduce disk writes on filesystems that support it (incl. ReFS & btrfs)

diff --git a/Engine/Source/Programs/Unsync/Private/UnsyncCore.cpp b/Engine/Source/Programs/Unsync/Private/UnsyncCore.cpp
--- a/Engine/Source/Programs/Unsync/Private/UnsyncCore.cpp
+++ b/Engine/Source/Programs/Unsync/Private/UnsyncCore.cpp
@@ -329,6 +329,15 @@ GetBuildTargetParams(const FSyncFileOptions& Options)
 	return Result;
 }
 
+FFileSyncResult&
+FFileSyncResult::operator+=(const FBuildTargetResult& BuildTargetResult)
+{
+	SourceBytes += BuildTargetResult.SourceBytes;
+	BaseBytes += BuildTargetResult.BaseBytes;
+	InSituBytes += BuildTargetResult.InSituBytes;
+	return *this;
+}
+
 FFileSyncResult
 SyncFile(const FNeedList&		   NeedList,
 		 const FPath&			   SourceFilePath,
@@ -442,8 +451,7 @@ SyncFile(const FNeedList&		   NeedList,
 		FBuildTargetParams BuildParams = GetBuildTargetParams(Options);
 		FBuildTargetResult BuildResult = BuildTarget(*TargetFile, SourceFile, BaseDataReader, NeedList, BuildParams);
 
-		Result.SourceBytes = BuildResult.SourceBytes;
-		Result.BaseBytes   = BuildResult.BaseBytes;
+		Result += BuildResult;
 
 		if (!BuildResult.bSuccess)
 		{
@@ -1487,6 +1495,7 @@ SyncDirectory(const FSyncDirectoryOptions& SyncOptions)
 
 	std::atomic<uint64> StatSourceBytes = {};
 	std::atomic<uint64> StatBaseBytes	= {};
+	std::atomic<uint64> StatInSituBytes	= {};
 
 	std::vector<FFileSyncTask> AllFileTasks;
 
@@ -1953,6 +1962,7 @@ SyncDirectory(const FSyncDirectoryOptions& SyncOptions)
 								 &SyncOptions,
 								 &StatSourceBytes,
 								 &StatBaseBytes,
+								 &StatInSituBytes,
 								 &BackgroundTaskStatMutex,
 								 &BackgroundTaskResults,
 								 &NumFailedTasks,
@@ -1996,6 +2006,7 @@ SyncDirectory(const FSyncDirectoryOptions& SyncOptions)
 			{
 				StatSourceBytes += SyncResult.SourceBytes;
 				StatBaseBytes += SyncResult.BaseBytes;
+				StatInSituBytes += SyncResult.InSituBytes;
 				UNSYNC_ASSERT(SyncResult.SourceBytes + SyncResult.BaseBytes == Item.TotalSizeBytes);
 
 				if (!GDryRun)
@@ -2222,7 +2233,7 @@ SyncDirectory(const FSyncDirectoryOptions& SyncOptions)
 	}
 
 	UNSYNC_LOG(L"Skipped files: %d, full copies: %d, partial copies: %d", StatSkipped, StatFullCopy, StatPartialCopy);
-	UNSYNC_LOG(L"Copied from source: %.2f MB, copied from base: %.2f MB", SizeMb(StatSourceBytes), SizeMb(StatBaseBytes));
+	UNSYNC_LOG(L"Copied from source: %.2f MB, copied from base: %.2f MB (including %.2f MB in situ)", SizeMb(StatSourceBytes), SizeMb(StatBaseBytes), SizeMb(StatInSituBytes));
 	UNSYNC_LOG(L"Sync completed %ls", bSyncSucceeded ? L"successfully" : L"with errors (see log for details)");
 
 	double ElapsedSeconds = DurationSec(TimeBegin, TimePointNow());
diff --git a/Engine/Source/Programs/Unsync/Private/UnsyncCore.h b/Engine/Source/Programs/Unsync/Private/UnsyncCore.h
--- a/Engine/Source/Programs/Unsync/Private/UnsyncCore.h
+++ b/Engine/Source/Programs/Unsync/Private/UnsyncCore.h
@@ -28,6 +28,7 @@ class FProxy;
 class FProxyPool;
 class FBlockCache;
 class FScavengeDatabase;
+struct FBuildTargetResult;
 struct FRemoteDesc;
 struct FIOReader;
 struct FIOWriter;
@@ -203,8 +204,11 @@ struct FFileSyncResult
 
 	bool Succeeded() const { return (uint32)Status < (uint32)EFileSyncStatus::ErrorUnknown; }
 
+	FFileSyncResult& operator+=(const FBuildTargetResult& BuildTargetResult);
+
 	uint64 SourceBytes = 0;
 	uint64 BaseBytes   = 0;
+	uint64 InSituBytes = 0;
 };
 
 struct FSyncFileOptions
diff --git a/Engine/Source/Programs/Unsync/Private/UnsyncFile.cpp b/Engine/Source/Programs/Unsync/Private/UnsyncFile.cpp
--- a/Engine/Source/Programs/Unsync/Private/UnsyncFile.cpp
+++ b/Engine/Source/Programs/Unsync/Private/UnsyncFile.cpp
@@ -256,6 +256,17 @@ GetCachedFileAttrib(const FPath& Path, FFileAttributeCache& AttribCache)
 	return Result;
 }
 
+int64
+FIOWriter::WriteFrom(FIOReader& Reader, const FCopyCommand& CopyCommand, FBufferView FileData, uint64* NumBytesClonedPtr)
+{
+	if (FileData.Size != CopyCommand.Size)
+	{
+		return -1;
+	}
+
+	return Write(FileData.Data, CopyCommand.TargetOffset, FileData.Size);
+}
+
 #if UNSYNC_PLATFORM_WINDOWS
 inline uint64
 MakeU64(FILETIME Ft)
@@ -856,6 +867,162 @@ FWindowsFile::CreateAsyncReader(uint32 MaxPipelineDepth)
 	return std::unique_ptr<FAsyncReader>(new FWindowsAsyncFileReader(*this, MaxPipelineDepth));
 }
 
+bool
+CloneFileRange(FWindowsFile& SourceFile, FWindowsFile& DestFile, const FCopyCommand& CopyCommand, FBufferView FileData)
+{
+	// ReFS has a number of preconditions for block cloning. See
+	// https://learn.microsoft.com/en-us/windows-server/storage/refs/block-cloning#functionality-restrictions-and-remarks
+
+	// Block clones must be aligned on cluster boundaries, typically 4 kiB.
+	// TODO: Query the filesystem cluster size, since ReFS also supports 64 kiB clusters.
+	constexpr uint64 AlignBytes = 4096;
+	if (CopyCommand.SourceOffset % AlignBytes != CopyCommand.TargetOffset % AlignBytes)
+	{
+		return false;
+	}
+
+	const uint64 PreBytes	= (AlignBytes - (CopyCommand.SourceOffset % AlignBytes)) % AlignBytes;
+	const uint64 PostBytes	= (CopyCommand.SourceOffset + CopyCommand.Size) % AlignBytes;
+	const uint64 CloneBytes = (CopyCommand.Size - PreBytes - PostBytes);
+	UNSYNC_ASSERT(CloneBytes % AlignBytes == 0);
+	UNSYNC_ASSERT(PreBytes + CloneBytes + PostBytes == CopyCommand.Size);
+	if (CloneBytes == 0)
+	{
+		return false;
+	}
+	const uint64 SourceCloneOffset = CopyCommand.SourceOffset + PreBytes;
+	const uint64 DestCloneOffset   = CopyCommand.TargetOffset + PreBytes;
+
+	// Check whether the files are on the same filesystem, and the filesystem supports block cloning.
+	{
+		DWORD SourceVolumeSerialNumber = 0;
+		DWORD DestVolumeSerialNumber   = 0;
+		DWORD DestVolumeFlags		   = 0;
+
+		if (GetVolumeInformationByHandleW(SourceFile.FileHandle,
+										  nullptr,
+										  0,
+										  &SourceVolumeSerialNumber,
+										  nullptr,
+										  nullptr,
+										  nullptr,
+										  0) != TRUE ||
+			GetVolumeInformationByHandleW(DestFile.FileHandle,
+										  nullptr,
+										  0,
+										  &DestVolumeSerialNumber,
+										  nullptr,
+										  &DestVolumeFlags,
+										  nullptr,
+										  0) != TRUE)
+		{
+			return false;
+		}
+
+		if (SourceVolumeSerialNumber != DestVolumeSerialNumber || (DestVolumeFlags & FILE_SUPPORTS_BLOCK_REFCOUNTING) == 0)
+		{
+			return false;
+		}
+	}
+
+	// Ensure target file is sparse if input file is.
+	{
+		FILE_BASIC_INFO SourceFileInfo{};
+		if (GetFileInformationByHandleEx(SourceFile.FileHandle, FileBasicInfo, &SourceFileInfo, sizeof(SourceFileInfo)))
+		{
+			FILE_BASIC_INFO DestFileInfo{};
+			if (GetFileInformationByHandleEx(DestFile.FileHandle, FileBasicInfo, &DestFileInfo, sizeof(DestFileInfo)))
+			{
+				if ((SourceFileInfo.FileAttributes & FILE_ATTRIBUTE_SPARSE_FILE) !=
+					(DestFileInfo.FileAttributes & FILE_ATTRIBUTE_SPARSE_FILE))
+				{
+					FILE_SET_SPARSE_BUFFER Sparse{};
+					Sparse.SetSparse = (SourceFileInfo.FileAttributes & FILE_ATTRIBUTE_SPARSE_FILE) ? TRUE : FALSE;
+					DWORD BytesReturned{};
+					if (!DeviceIoControl(DestFile.FileHandle,
+										 FSCTL_SET_SPARSE,
+										 &Sparse,
+										 sizeof(Sparse),
+										 nullptr,
+										 0,
+										 &BytesReturned,
+										 nullptr))
+					{
+						UNSYNC_VERBOSE("Unable to set target file sparseness to match source file: %s",
+									   std::system_category().message(GetLastError()).c_str());
+					}
+				}
+			}
+		}
+	}
+
+	// Ensure the output file is big enough
+	{
+		LARGE_INTEGER ExistingFileSize{};
+		if (!GetFileSizeEx(DestFile.FileHandle, &ExistingFileSize))
+		{
+			return false;
+		}
+
+		const int64 NewMinSize = CopyCommand.TargetOffset + CopyCommand.Size;
+		if (ExistingFileSize.QuadPart < NewMinSize)
+		{
+			FILE_END_OF_FILE_INFO EOFInfo{};
+			EOFInfo.EndOfFile.QuadPart = NewMinSize;
+			if (!SetFileInformationByHandle(DestFile.FileHandle, FileEndOfFileInfo, &EOFInfo, sizeof(EOFInfo)))
+			{
+				return false;
+			}
+		}
+	}
+
+	DUPLICATE_EXTENTS_DATA Extent{};
+	Extent.FileHandle				 = SourceFile.FileHandle;
+	Extent.SourceFileOffset.QuadPart = SourceCloneOffset;
+	Extent.TargetFileOffset.QuadPart = DestCloneOffset;
+	Extent.ByteCount.QuadPart		 = CloneBytes;
+	DWORD	BytesReturned{};
+	if (DeviceIoControl(DestFile.FileHandle, FSCTL_DUPLICATE_EXTENTS_TO_FILE, &Extent, sizeof(Extent), NULL, 0, &BytesReturned, NULL) != TRUE)
+	{
+		return false;
+	}
+
+	// Write unaligned pre/post data
+	if (FileData.Size == CopyCommand.Size)
+	{
+		// Use file data that was already read into memory.
+		const uint64 PostOffset = PreBytes + CloneBytes;
+		if (DestFile.Write(FileData.Data, CopyCommand.TargetOffset, PreBytes) != PreBytes ||
+			DestFile.Write(FileData.Data + PostOffset, CopyCommand.TargetOffset + PostOffset, PostBytes) != PostBytes)
+		{
+			return false;
+		}
+	}
+	else
+	{
+		// It's preferable for callers to pass in the already-in-memory data via `FileData` rather than re-reading it here.
+		uint8 Buffer[AlignBytes];
+		if (PreBytes > 0)
+		{
+			if (SourceFile.Read(Buffer, CopyCommand.SourceOffset, PreBytes) != PreBytes ||
+				DestFile.Write(Buffer, CopyCommand.TargetOffset, PreBytes) != PreBytes)
+			{
+				return false;
+			}
+		}
+		if (PostBytes > 0)
+		{
+			if (SourceFile.Read(Buffer, CopyCommand.SourceOffset + PreBytes + CloneBytes, PostBytes) != PostBytes ||
+				DestFile.Write(Buffer, CopyCommand.TargetOffset + PreBytes + CloneBytes, PostBytes) != PostBytes)
+			{
+				return false;
+			}
+		}
+	}
+
+	return true;
+}
+
 FFileAttributes
 GetFileAttrib(const FPath& Path, FFileAttributeCache* AttribCache)
 {
@@ -952,6 +1119,7 @@ FUnixFile::Close()
 	{
 		fclose(FileHandle);
 		FileHandle = nullptr;
+		FileDescriptor = -1;
 	}
 }
 
@@ -981,6 +1149,29 @@ FUnixFile::Write(const void* data, uint64 DestOffset, uint64 WriteSize)
 	return wrote_bytes;
 }
 
+bool
+CloneFileRange(FUnixFile& SourceFile, FUnixFile& DestFile, const FCopyCommand& CopyCommand, FBufferView FileData)
+{
+	off_t SourceOffset = CopyCommand.SourceOffset;
+	off_t DestOffset = CopyCommand.TargetOffset;
+
+	ssize_t Written = copy_file_range(
+		SourceFile.FileDescriptor,
+		&SourceOffset,
+		DestFile.FileDescriptor,
+		&DestOffset,
+		CopyCommand.Size,
+		0);
+
+	if (Written == -1)
+	{
+		UNSYNC_VERBOSE2("copy_file_range failed: %s", strerror(errno));
+		return false;
+	}
+
+	return Written == CopyCommand.Size;
+}
+
 FFileAttributes
 GetFileAttrib(const FPath& Path, FFileAttributeCache* AttribCache)
 {
diff --git a/Engine/Source/Programs/Unsync/Private/UnsyncFile.h b/Engine/Source/Programs/Unsync/Private/UnsyncFile.h
--- a/Engine/Source/Programs/Unsync/Private/UnsyncFile.h
+++ b/Engine/Source/Programs/Unsync/Private/UnsyncFile.h
@@ -3,6 +3,7 @@
 #pragma once
 
 #include "UnsyncBuffer.h"
+#include "UnsyncCore.h"
 #include "UnsyncUtil.h"
 
 #include <algorithm>
@@ -158,6 +159,9 @@ struct FIOReader : virtual FIOBase
 	virtual uint64 Read(void* Dest, uint64 SourceOffset, uint64 Size) = 0;
 
 	virtual std::unique_ptr<FAsyncReader> CreateAsyncReader(uint32 MaxPipelineDepth = MAX_IO_PIPELINE_DEPTH);
+
+	// Get the actual reader, eg. for a deferred reader that has an inner file reader.
+	virtual FIOReader& ActualReader() { return *this; }
 };
 
 struct FDummyAsyncReader final : FAsyncReader
@@ -176,18 +180,58 @@ private:
 struct FIOWriter : virtual FIOBase
 {
 	virtual uint64 Write(const void* Data, uint64 DestOffset, uint64 Size) = 0;
+
+	// Write file data from Reader. File-backed writers may use block cloning to copy the data by reference.
+	// FileData.Size must equal CopyCommand.Size.
+	// NumBytesClonedPtr may be null. Otherwise the number of bytes cloned is added to the target variable.
+	virtual int64 WriteFrom(FIOReader& Reader, const FCopyCommand& CopyCommand, FBufferView FileData, uint64* NumBytesClonedPtr = nullptr);
 };
 
 struct FIOReaderWriter : FIOReader, FIOWriter
 {
 };
 
+// File reader/writer that uses CRTP to implement block-cloned copies.
+template <typename InFileType>
+class FFileReaderWriter : public FIOReaderWriter
+{
+public:
+	using FileType = InFileType;
+
+	virtual int64 WriteFrom(FIOReader& Reader, const FCopyCommand& CopyCommand, FBufferView FileData, uint64* NumBytesClonedPtr) override
+	{
+		FileType* FileReader = dynamic_cast<FileType*>(&Reader.ActualReader());
+		if (FileReader)
+		{
+			// Attempt block cloning
+			FileType* FileWriter = dynamic_cast<FileType*>(this);
+			UNSYNC_ASSERT(FileWriter);
+
+			if (CloneFileRange(*FileReader, *FileWriter, CopyCommand, FileData))
+			{
+				if (NumBytesClonedPtr)
+				{
+					*NumBytesClonedPtr += CopyCommand.Size;
+				}
+				return CopyCommand.Size;
+			}
+		}
+
+		if (FileData.Size != CopyCommand.Size)
+		{
+			return -1;
+		}
+
+		return Write(FileData.Data, CopyCommand.TargetOffset, FileData.Size);
+	}
+};
+
 #if UNSYNC_PLATFORM_WINDOWS
-struct FWindowsFile : FIOReaderWriter
+struct FWindowsFile : FFileReaderWriter<FWindowsFile>
 {
 	friend struct FWindowsAsyncFileReader;
 
-	FWindowsFile(const FPath& Filename, EFileMode Mode = EFileMode::ReadOnly, uint64 InSize = 0);
+	explicit FWindowsFile(const FPath& Filename, EFileMode Mode = EFileMode::ReadOnly, uint64 InSize = 0);
 	~FWindowsFile();
 
 	// IOBase
@@ -226,13 +270,24 @@ private:
 	void								  RemoveAsyncReader(FWindowsAsyncFileReader* Reader);
 	std::vector<FWindowsAsyncFileReader*> AsyncReaders;
 };
+
+// Clone file blocks via ref counting if possible.
+// FileData optionally supplies an in-memory view of the file contents that's used for writing unaligned bytes at the beginning and end of the cloned range.
+bool CloneFileRange(FWindowsFile& SourceFile, FWindowsFile& DestFile, const FCopyCommand& CopyCommand, FBufferView FileData = FBufferView());
+
 using FNativeFile = FWindowsFile;
 #endif	// UNSYNC_PLATFORM_WINDOWS
 
 #if UNSYNC_PLATFORM_UNIX
-struct FUnixFile : FIOReaderWriter
+struct FUnixFile;
+
+// Copy file data using in-kernel copy, which also provides block cloning (reflinking) on filesystems that support it.
+// The `FileData` parameter is not used; it's only for API compatiblity with the Windows CloneFileRange() function.
+bool CloneFileRange(FUnixFile& SourceFile, FUnixFile& DestFile, const FCopyCommand& CopyCommand, FBufferView FileData = FBufferView());
+
+struct FUnixFile : FFileReaderWriter<FUnixFile>
 {
-	FUnixFile(const FPath& InFilename, EFileMode InMode = EFileMode::ReadOnly, uint64 InSize = 0);
+	explicit FUnixFile(const FPath& InFilename, EFileMode InMode = EFileMode::ReadOnly, uint64 InSize = 0);
 	~FUnixFile();
 
 	// IOBase
@@ -254,13 +309,15 @@ struct FUnixFile : FIOReaderWriter
 
 	static constexpr uint32 UNBUFFERED_READ_ALIGNMENT = 4096;
 
+	friend bool CloneFileRange(FUnixFile& SourceFile, FUnixFile& DestFile, const FCopyCommand& CopyCommand, FBufferView FileData);
+
 private:
 	bool OpenFileHandle(EFileMode InMode);
 
 private:
-	EFileMode Mode;
+	EFileMode Mode			 = EFileMode::None;
 	FILE*	  FileHandle	 = nullptr;
-	int		  FileDescriptor = 0;
+	int		  FileDescriptor = -1;
 };
 using FNativeFile = FUnixFile;
 #endif	// UNSYNC_PLATFORM_UNIX
@@ -394,6 +451,7 @@ struct FDeferredOpenReader : FIOReader
 	{
 		return GetOrOpenInner()->CreateAsyncReader(MaxPipelineDepth);
 	}
+	virtual FIOReader& ActualReader() override { return *GetOrOpenInner(); }
 
 	FIOReader* GetOrOpenInner()
 	{
diff --git a/Engine/Source/Programs/Unsync/Private/UnsyncScavenger.cpp b/Engine/Source/Programs/Unsync/Private/UnsyncScavenger.cpp
--- a/Engine/Source/Programs/Unsync/Private/UnsyncScavenger.cpp
+++ b/Engine/Source/Programs/Unsync/Private/UnsyncScavenger.cpp
@@ -401,7 +401,7 @@ BuildTargetFromScavengedData(FIOWriter&						Output,
 
 			bool bFoundInvalidBlock = false;
 
-			auto ReadCallback = [&Output, &OutScavengedBlocks, &ProgressLogger, &BuildResult, &bFoundInvalidBlock, StrongHasher](
+			auto ReadCallback = [&Output, &OutScavengedBlocks, &LocalSourceFile, &ProgressLogger, &BuildResult, &bFoundInvalidBlock, StrongHasher](
 									FIOBuffer Buffer,
 									uint64	  SourceOffset,
 									uint64	  ReadSize,
@@ -438,7 +438,7 @@ BuildTargetFromScavengedData(FIOWriter&						Output,
 
 				if (bBlockHashOk)
 				{
-					Output.Write(Buffer.GetData(), FirstBlock.TargetOffset, ReadSize);
+					Output.WriteFrom(LocalSourceFile, CopyCommand, Buffer.GetBufferView(), &BuildResult.InSituBytes);
 
 					for (const FNeedBlock& Block : CopyCommand.BlockRange)
 					{
diff --git a/Engine/Source/Programs/Unsync/Private/UnsyncScavenger.h b/Engine/Source/Programs/Unsync/Private/UnsyncScavenger.h
--- a/Engine/Source/Programs/Unsync/Private/UnsyncScavenger.h
+++ b/Engine/Source/Programs/Unsync/Private/UnsyncScavenger.h
@@ -76,6 +76,8 @@ private:
 struct FScavengedBuildTargetResult
 {
 	uint64 ScavengedBytes = 0;
+	// Bytes reused on disk via block cloning. Subset of ScavengedBytes.
+	uint64 InSituBytes	  = 0;
 };
 
 FScavengedBuildTargetResult BuildTargetFromScavengedData(FIOWriter&						Output,
diff --git a/Engine/Source/Programs/Unsync/Private/UnsyncTarget.cpp b/Engine/Source/Programs/Unsync/Private/UnsyncTarget.cpp
--- a/Engine/Source/Programs/Unsync/Private/UnsyncTarget.cpp
+++ b/Engine/Source/Programs/Unsync/Private/UnsyncTarget.cpp
@@ -212,6 +212,9 @@ BuildTarget(FIOWriter& Output, FIOReader& Source, FIOReader& Base, const FNeedLi
 	{
 		std::atomic<uint64> WrittenBytesFromSource = {};
 		std::atomic<uint64> WrittenBytesFromBase   = {};
+		// Bytes reused on disk via block cloning (if available).
+		// Subset of WrittenBytesFromBase.
+		std::atomic<uint64> BytesUsedInSitu		   = {};
 	};
 	FStats Stats;
 
@@ -241,6 +244,7 @@ BuildTarget(FIOWriter& Output, FIOReader& Source, FIOReader& Base, const FNeedLi
 
 		// Treat scavenged data the same as base for stats purposes
 		Stats.WrittenBytesFromBase += ScavengeResult.ScavengedBytes;
+		Stats.BytesUsedInSitu += ScavengeResult.InSituBytes;
 
 		FilteredSourceNeedList = NeedList.Source;
 		auto FilterResult	   = std::remove_if(FilteredSourceNeedList.begin(),
@@ -350,15 +354,17 @@ BuildTarget(FIOWriter& Output, FIOReader& Source, FIOReader& Base, const FNeedLi
 
 			uint64 ReadBytes = 0;
 
-			auto ReadCallback = [&ReadBytes, &Output, &WriteTasks, &Error, &Stats, Block, ListType](FIOBuffer CmdBuffer,
-																									uint64	  CmdOffset,
-																									uint64	  CmdReadSize,
-																									uint64	  CmdUserData)
+			auto ReadCallback = [&DataProvider, &ReadBytes, &Output, &WriteTasks, &Error, &Stats, Block, ListType](FIOBuffer CmdBuffer,
+																												   uint64	  CmdOffset,
+																												   uint64	  CmdReadSize,
+																												   uint64	  CmdUserData)
 			{
 				WriteTasks.run(
-					[Buffer = MakeShared(std::move(CmdBuffer)), CmdReadSize, Block, &Output, &Error, &Stats, ListType]()
+					[Buffer = MakeShared(std::move(CmdBuffer)), CmdReadSize, Block, &DataProvider, &Output, &Error, &Stats, ListType]()
 					{
-						const uint64 WrittenBytes = Output.Write(Buffer->GetData(), Block.TargetOffset, CmdReadSize);
+						uint64		InSituBytes	 = 0;
+						const int64 WrittenBytes = Output.WriteFrom(DataProvider, Block, Buffer->GetBufferView(), &InSituBytes);
+						Stats.BytesUsedInSitu += InSituBytes;
 
 						if (ListType == EBlockListType::Source)
 						{
@@ -373,7 +379,7 @@ BuildTarget(FIOWriter& Output, FIOReader& Source, FIOReader& Base, const FNeedLi
 							UNSYNC_FATAL(L"Unexpected block list type");
 						}
 
-						if (WrittenBytes != CmdReadSize)
+						if (WrittenBytes < 0 || uint64(WrittenBytes) != CmdReadSize)
 						{
 							UNSYNC_FATAL(L"Expected to write %llu bytes, but written %llu", CmdReadSize, WrittenBytes);
 							Error.Set(AppError(L"Failed to write output"));
@@ -722,6 +728,10 @@ BuildTarget(FIOWriter& Output, FIOReader& Source, FIOReader& Base, const FNeedLi
 
 	ElapsedTime += DurationSec(TimeBegin, TimePointNow());
 	UNSYNC_VERBOSE(L"Done in %.3f sec (%.3f MB / sec)", ElapsedTime, SizeMb(double(SizeInfo.TotalBytes) / ElapsedTime));
+	if (Stats.BytesUsedInSitu > 0)
+	{
+		UNSYNC_VERBOSE2("%.2f MB reused in situ", SizeMb(Stats.BytesUsedInSitu));
+	}
 
 	if (GLogVerbose)
 	{
@@ -731,6 +741,7 @@ BuildTarget(FIOWriter& Output, FIOReader& Source, FIOReader& Base, const FNeedLi
 	BuildResult.bSuccess	= !Error;
 	BuildResult.BaseBytes	= Stats.WrittenBytesFromBase;
 	BuildResult.SourceBytes = Stats.WrittenBytesFromSource;
+	BuildResult.InSituBytes = Stats.BytesUsedInSitu;
 
 	return BuildResult;
 }
diff --git a/Engine/Source/Programs/Unsync/Private/UnsyncTarget.h b/Engine/Source/Programs/Unsync/Private/UnsyncTarget.h
--- a/Engine/Source/Programs/Unsync/Private/UnsyncTarget.h
+++ b/Engine/Source/Programs/Unsync/Private/UnsyncTarget.h
@@ -18,6 +18,7 @@ struct FBuildTargetResult
 	bool   bSuccess	   = false;
 	uint64 SourceBytes = 0;
 	uint64 BaseBytes   = 0;
+	uint64 InSituBytes = 0;
 };
 
 struct FBuildTargetParams
diff --git a/Engine/Source/Programs/Unsync/Private/UnsyncTest.cpp b/Engine/Source/Programs/Unsync/Private/UnsyncTest.cpp
--- a/Engine/Source/Programs/Unsync/Private/UnsyncTest.cpp
+++ b/Engine/Source/Programs/Unsync/Private/UnsyncTest.cpp
@@ -21,6 +21,40 @@ UNSYNC_THIRD_PARTY_INCLUDES_END
 
 namespace unsync {
 
+namespace test {
+
+class FTestFile : public FNativeFile
+{
+public:
+	explicit FTestFile(const FPath& Filename, EFileMode Mode = EFileMode::ReadOnly, uint64 InSize = 0)
+		: FNativeFile(Filename, Mode, InSize)
+	{
+#if UNSYNC_PLATFORM_WINDOWS
+		if (DWORD FileAttributes = GetFileAttributesW(Filename.c_str()); FileAttributes != INVALID_FILE_ATTRIBUTES)
+		{
+			FileAttributes |= FILE_ATTRIBUTE_TEMPORARY;
+			SetFileAttributesW(Filename.c_str(), FileAttributes);
+		}
+#endif
+	}
+
+	virtual ~FTestFile() override {
+		if (!Filename.empty())
+		{
+			Close();
+
+			std::error_code Error{};
+			std::filesystem::remove(Filename.native(), Error);
+			if (Error)
+			{
+				UNSYNC_ERROR("Failed to remove test file %s: %s", Filename.c_str(), Error.message().c_str());
+			}
+		}
+	}
+};
+
+} // namespace test
+
 static bool
 VerifyRandomBytes(const uint8* Data, uint64 Size, uint32 Seed)
 {
@@ -463,12 +497,17 @@ TestBuildTarget(EWeakHashAlgorithmID WeakHasher, EStrongHashAlgorithmID StrongHa
 void
 TestFiles()
 {
+	using test::FTestFile;
+
 	UNSYNC_LOG(L"TestFiles()");
 
 	// TODO: test different chunk modes
 	const EChunkingAlgorithmID ChunkingMode = EChunkingAlgorithmID::FixedBlocks;
 
 	FPath TempDirPath = std::filesystem::temp_directory_path() / "unsync_test";
+	std::error_code ErrorCode{};
+	// Whether to remove the temp dir after running file tests.
+	const bool bRemoveTempDir = !std::filesystem::exists(TempDirPath, ErrorCode);
 	CreateDirectories(TempDirPath);
 
 	bool bDirectoryExists = PathExists(TempDirPath) && IsDirectory(TempDirPath);
@@ -477,7 +516,7 @@ TestFiles()
 	FPath  TestFilename = TempDirPath / "a.bin";
 	uint64 TestFileSize = 32'489'595;
 
-	std::unique_ptr<uint8> TempBuffer = std::unique_ptr<uint8>(new uint8[TestFileSize]);
+	std::unique_ptr<uint8[]> TempBuffer = std::make_unique<uint8[]>(TestFileSize);
 
 	{
 		// overlapped file writing
@@ -494,6 +533,61 @@ TestFiles()
 		VerifyRandomBytes(TempBuffer.get(), TestFileSize, 321);
 	}
 
+	// CloneFileRange
+	{
+		constexpr uint64 CloneFileSize = 8 * 1024 * 1024;
+		std::unique_ptr<uint8[]> CloneData = std::make_unique<uint8[]>(CloneFileSize);
+		UNSYNC_ASSERT(CloneData);
+		FTestFile InputFile(TempDirPath / "input.bin", EFileMode::CreateReadWrite, CloneFileSize);
+		UNSYNC_ASSERT(InputFile.IsValid());
+		FTestFile OutputFile(TempDirPath / "output.bin", EFileMode::CreateReadWrite, CloneFileSize);
+		UNSYNC_ASSERT(OutputFile.IsValid());
+
+		const int Seed = __COUNTER__;
+		FillRandomBytes(CloneData.get(), CloneFileSize, Seed);
+		uint64 Written = InputFile.Write(CloneData.get(), 0, CloneFileSize);
+		UNSYNC_ASSERT(Written == CloneFileSize);
+
+		// Copy the file with the 1st/2nd halves swapped.
+		{
+			const FCopyCommand Copy1{.Size = CloneFileSize / 2, .SourceOffset = 0, .TargetOffset = CloneFileSize / 2};
+			const FCopyCommand Copy2{.Size = CloneFileSize / 2, .SourceOffset = CloneFileSize / 2, .TargetOffset = 0};
+
+			if (CloneFileRange(InputFile, OutputFile, Copy1) && CloneFileRange(InputFile, OutputFile, Copy2))
+			{
+				UNSYNC_VERBOSE("Successfully cloned file contents");
+				memset(CloneData.get(), 0, CloneFileSize);
+				OutputFile.Read(CloneData.get(), CloneFileSize / 2, CloneFileSize / 2);
+				OutputFile.Read(CloneData.get() + CloneFileSize / 2, 0, CloneFileSize / 2);
+				VerifyRandomBytes(CloneData.get(), CloneFileSize, Seed);
+			}
+			else
+			{
+				UNSYNC_VERBOSE("Unable to clone file contents—probably an unsupported filesystem.");
+			}
+		}
+
+		// Non-cluster-aligned clone
+		{
+			constexpr int	   UnalignedOffset		 = 3;
+			constexpr int	   UnalignedCloneSize	 = CloneFileSize / 2 - UnalignedOffset;
+			constexpr int	   UnalignedSourceOffset = UnalignedOffset;
+			constexpr int	   UnalignedTargetOffset = CloneFileSize / 2 + UnalignedOffset;
+			const FCopyCommand CopyUnaligned{.Size		   = UnalignedCloneSize,
+											 .SourceOffset = UnalignedSourceOffset,
+											 .TargetOffset = UnalignedTargetOffset};
+			if (CloneFileRange(InputFile, OutputFile, CopyUnaligned))
+			{
+				UNSYNC_VERBOSE("Successfully cloned unaligned file contents");
+				std::unique_ptr<uint8[]> UnalignedReadData = std::make_unique<uint8[]>(UnalignedCloneSize);
+				UNSYNC_ASSERT(UnalignedReadData);
+				const uint64 ReadLen = OutputFile.Read(UnalignedReadData.get(), UnalignedOffset, UnalignedCloneSize);
+				UNSYNC_ASSERT(ReadLen == UnalignedCloneSize);
+				UNSYNC_ASSERT(memcmp(UnalignedReadData.get(), CloneData.get() + UnalignedOffset, UnalignedCloneSize) == 0);
+			}
+		}
+	}
+
 	{
 		memset(TempBuffer.get(), 0, TestFileSize);
 
@@ -541,6 +635,15 @@ TestFiles()
 		UNSYNC_ASSERT(HashU32x4[2] == 0xa9658163);
 		UNSYNC_ASSERT(HashU32x4[3] == 0x2c317d4b);
 	}
+
+	if (bRemoveTempDir)
+	{
+		std::filesystem::remove_all(TempDirPath, ErrorCode);
+		if (ErrorCode)
+		{
+			UNSYNC_LOG("Failed to remove TestFiles temp dir %s: %s", TempDirPath.c_str(), ErrorCode.message().c_str());
+		}
+	}
 }
 
 void
diff --git a/Engine/Source/Programs/Unsync/Private/UnsyncVersion.cpp b/Engine/Source/Programs/Unsync/Private/UnsyncVersion.cpp
--- a/Engine/Source/Programs/Unsync/Private/UnsyncVersion.cpp
+++ b/Engine/Source/Programs/Unsync/Private/UnsyncVersion.cpp
@@ -2,7 +2,7 @@
 
 #include "UnsyncVersion.h"
 
-#define UNSYNC_VERSION_STR "1.0.99"
+#define UNSYNC_VERSION_STR "1.1.0"
 
 namespace unsync {
 
