Unsync: Fix copy-on-write handling of small ranges on Windows

#tests Unit tests & sync published builds
#rnx

diff --git a/Engine/Source/Programs/Unsync/Private/UnsyncFile.cpp b/Engine/Source/Programs/Unsync/Private/UnsyncFile.cpp
--- a/Engine/Source/Programs/Unsync/Private/UnsyncFile.cpp
+++ b/Engine/Source/Programs/Unsync/Private/UnsyncFile.cpp
@@ -883,13 +883,14 @@ CloneFileRange(FWindowsFile& SourceFile, FWindowsFile& DestFile, const FCopyComm
 
 	const uint64 PreBytes	= (AlignBytes - (CopyCommand.SourceOffset % AlignBytes)) % AlignBytes;
 	const uint64 PostBytes	= (CopyCommand.SourceOffset + CopyCommand.Size) % AlignBytes;
-	const uint64 CloneBytes = (CopyCommand.Size - PreBytes - PostBytes);
-	UNSYNC_ASSERT(CloneBytes % AlignBytes == 0);
-	UNSYNC_ASSERT(PreBytes + CloneBytes + PostBytes == CopyCommand.Size);
-	if (CloneBytes == 0)
+	// Check for writes that are too small to clone any blocks.
+	if (PreBytes >= CopyCommand.Size || PostBytes >= CopyCommand.Size || CopyCommand.Size < PreBytes + PostBytes + AlignBytes)
 	{
 		return false;
 	}
+	const uint64 CloneBytes = (CopyCommand.Size - PreBytes - PostBytes);
+	UNSYNC_ASSERT(CloneBytes % AlignBytes == 0);
+	UNSYNC_ASSERT(PreBytes + CloneBytes + PostBytes == CopyCommand.Size);
 	const uint64 SourceCloneOffset = CopyCommand.SourceOffset + PreBytes;
 	const uint64 DestCloneOffset   = CopyCommand.TargetOffset + PreBytes;
 
diff --git a/Engine/Source/Programs/Unsync/Private/UnsyncTarget.cpp b/Engine/Source/Programs/Unsync/Private/UnsyncTarget.cpp
--- a/Engine/Source/Programs/Unsync/Private/UnsyncTarget.cpp
+++ b/Engine/Source/Programs/Unsync/Private/UnsyncTarget.cpp
@@ -366,6 +366,13 @@ BuildTarget(FIOWriter& Output, FIOReader& Source, FIOReader& Base, const FNeedLi
 						const int64 WrittenBytes = Output.WriteFrom(DataProvider, Block, Buffer->GetBufferView(), &InSituBytes);
 						Stats.BytesUsedInSitu += InSituBytes;
 
+						if (WrittenBytes < 0 || uint64(WrittenBytes) != CmdReadSize)
+						{
+							UNSYNC_FATAL(L"Expected to write %llu bytes, but written %llu", CmdReadSize, WrittenBytes);
+							Error.Set(AppError(L"Failed to write output"));
+							return;
+						}
+
 						if (ListType == EBlockListType::Source)
 						{
 							Stats.WrittenBytesFromSource += WrittenBytes;
@@ -378,12 +385,6 @@ BuildTarget(FIOWriter& Output, FIOReader& Source, FIOReader& Base, const FNeedLi
 						{
 							UNSYNC_FATAL(L"Unexpected block list type");
 						}
-
-						if (WrittenBytes < 0 || uint64(WrittenBytes) != CmdReadSize)
-						{
-							UNSYNC_FATAL(L"Expected to write %llu bytes, but written %llu", CmdReadSize, WrittenBytes);
-							Error.Set(AppError(L"Failed to write output"));
-						}
 					});
 
 				ReadBytes = CmdReadSize;
diff --git a/Engine/Source/Programs/Unsync/Private/UnsyncTest.cpp b/Engine/Source/Programs/Unsync/Private/UnsyncTest.cpp
--- a/Engine/Source/Programs/Unsync/Private/UnsyncTest.cpp
+++ b/Engine/Source/Programs/Unsync/Private/UnsyncTest.cpp
@@ -581,11 +581,26 @@ TestFiles()
 				UNSYNC_VERBOSE("Successfully cloned unaligned file contents");
 				std::unique_ptr<uint8[]> UnalignedReadData = std::make_unique<uint8[]>(UnalignedCloneSize);
 				UNSYNC_ASSERT(UnalignedReadData);
-				const uint64 ReadLen = OutputFile.Read(UnalignedReadData.get(), UnalignedOffset, UnalignedCloneSize);
+				const uint64 ReadLen = OutputFile.Read(UnalignedReadData.get(), UnalignedTargetOffset, UnalignedCloneSize);
 				UNSYNC_ASSERT(ReadLen == UnalignedCloneSize);
 				UNSYNC_ASSERT(memcmp(UnalignedReadData.get(), CloneData.get() + UnalignedOffset, UnalignedCloneSize) == 0);
 			}
 		}
+
+		// Unaligned source/dest. Not supported on ReFS.
+		{
+			const FCopyCommand UnalignedCopy{.Size = 16384, .SourceOffset = 1, .TargetOffset = 2};
+			(void)CloneFileRange(InputFile, OutputFile, UnalignedCopy);
+		}
+
+		// Small ranges; too small to block clone on ReFS (with 4 kiB clusters)
+		// OK for clones to fail.
+		{
+			(void)CloneFileRange(InputFile, OutputFile, FCopyCommand{.Size = 100, .SourceOffset = 200, .TargetOffset = 200});
+			(void)CloneFileRange(InputFile, OutputFile, FCopyCommand{.Size = 100, .SourceOffset = 4000, .TargetOffset = 4000});
+			(void)CloneFileRange(InputFile, OutputFile, FCopyCommand{.Size = 4096, .SourceOffset = 100, .TargetOffset = 100});
+			(void)CloneFileRange(InputFile, OutputFile, FCopyCommand{.Size = 5096, .SourceOffset = 100, .TargetOffset = 100});
+		}
 	}
 
 	{
diff --git a/Engine/Source/Programs/Unsync/Private/UnsyncVersion.cpp b/Engine/Source/Programs/Unsync/Private/UnsyncVersion.cpp
--- a/Engine/Source/Programs/Unsync/Private/UnsyncVersion.cpp
+++ b/Engine/Source/Programs/Unsync/Private/UnsyncVersion.cpp
@@ -2,7 +2,7 @@
 
 #include "UnsyncVersion.h"
 
-#define UNSYNC_VERSION_STR "1.1.0"
+#define UNSYNC_VERSION_STR "1.1.1"
 
 namespace unsync {
 
