Unsync: Fix garbled scavenge log on Linux

#rnx

diff --git a/Engine/Source/Programs/Unsync/Private/UnsyncScavenger.cpp b/Engine/Source/Programs/Unsync/Private/UnsyncScavenger.cpp
--- a/Engine/Source/Programs/Unsync/Private/UnsyncScavenger.cpp
+++ b/Engine/Source/Programs/Unsync/Private/UnsyncScavenger.cpp
@@ -31,7 +31,7 @@ FScavengeDatabase::BuildFromFileSyncTasks(const FSyncDirectoryOptions& SyncOptio
 
 	const FPath ExtendedTargetPath = RemoveExtendedPathPrefix(SyncOptions.Target);
 
-	UNSYNC_VERBOSE(L"Scanning '%ls' for usable manifests", SyncOptions.ScavengeRoot.c_str());
+	UNSYNC_VERBOSE(L"Scanning '%ls' for usable manifests", SyncOptions.ScavengeRoot.wstring().c_str());
 
 	for (std::filesystem::recursive_directory_iterator DirIt = RecursiveDirectoryScan(SyncOptions.ScavengeRoot);
 		 DirIt != std::filesystem::recursive_directory_iterator();
