Unsync: Reflink reused data blocks on Linux

#rnx

diff --git a/Engine/Source/Programs/Unsync/Private/UnsyncFile.cpp b/Engine/Source/Programs/Unsync/Private/UnsyncFile.cpp
--- a/Engine/Source/Programs/Unsync/Private/UnsyncFile.cpp
+++ b/Engine/Source/Programs/Unsync/Private/UnsyncFile.cpp
@@ -12,6 +12,7 @@
 
 #if UNSYNC_PLATFORM_UNIX
 #	include <errno.h>
+#	include <fcntl.h>
 #	include <sys/stat.h>
 #	include <sys/types.h>
 #	include <unistd.h>
@@ -33,6 +34,74 @@ static constexpr uint64 SECONDS_BETWEEN_WINDOWS_AND_UNIX = 11'644'473'600ull;
 static constexpr uint64 NANOS_PER_WINDOWS_TICK			 = 100ull;
 static constexpr uint64 WINDOWS_TICKS_PER_SECOND		 = 1'000'000'000ull / NANOS_PER_WINDOWS_TICK;  // each tick is 100ns
 
+struct FCloneCommand
+{
+	// Number of bytes preceding the aligned range
+	uint64 PreBytes = 0;
+
+	// Number of bytes in the aligned range.
+	uint64 CloneBytes = 0;
+
+	// Number of bytes after the aligned range.
+	uint64 PostBytes = 0;
+
+	uint64 SourcePreBytesOffset = 0;
+	uint64 DestPreBytesOffset = 0;
+
+	uint64 SourceCloneOffset = 0;
+	uint64 DestCloneOffset = 0;
+
+	uint64 SourcePostBytesOffset = 0;
+	uint64 DestPostBytesOffset = 0;
+
+	static std::optional<FCloneCommand> TryCreate(const FCopyCommand& CopyCommand, uint64 AlignBytes);
+};
+
+// Construct an FCloneCommand for the given alignment, or nullopt if the clone cannot be aligned.
+std::optional<FCloneCommand> FCloneCommand::TryCreate(const FCopyCommand& CopyCommand, uint64 AlignBytes)
+{
+	if (CopyCommand.SourceOffset % AlignBytes != CopyCommand.TargetOffset % AlignBytes)
+	{
+		return std::nullopt;
+	}
+
+	const uint64 PreBytes = (AlignBytes - (CopyCommand.SourceOffset % AlignBytes)) % AlignBytes;
+	const uint64 PostBytes = (CopyCommand.SourceOffset + CopyCommand.Size) % AlignBytes;
+
+	// Check for writes that are too small to clone any blocks.
+	if (PreBytes >= CopyCommand.Size || PostBytes >= CopyCommand.Size || CopyCommand.Size < PreBytes + PostBytes + AlignBytes)
+	{
+		return std::nullopt;
+	}
+
+	const uint64 CloneBytes = (CopyCommand.Size - PreBytes - PostBytes);
+
+	UNSYNC_ASSERT(CloneBytes % AlignBytes == 0);
+	UNSYNC_ASSERT(PreBytes + CloneBytes + PostBytes == CopyCommand.Size);
+
+	FCloneCommand CloneCommand{
+		.PreBytes = PreBytes,
+		.CloneBytes = CloneBytes,
+		.PostBytes = PostBytes,
+
+		.SourcePreBytesOffset = CopyCommand.SourceOffset,
+		.DestPreBytesOffset = CopyCommand.TargetOffset,
+
+		.SourceCloneOffset = CopyCommand.SourceOffset + PreBytes,
+		.DestCloneOffset = CopyCommand.TargetOffset + PreBytes,
+
+		.SourcePostBytesOffset = CopyCommand.SourceOffset + PreBytes + CloneBytes,
+		.DestPostBytesOffset = CopyCommand.TargetOffset + PreBytes + CloneBytes,
+	};
+
+	UNSYNC_ASSERT(CloneCommand.SourceCloneOffset % AlignBytes == 0);
+	UNSYNC_ASSERT(CloneCommand.DestCloneOffset % AlignBytes == 0);
+	UNSYNC_ASSERT(CloneCommand.SourcePostBytesOffset % AlignBytes == 0);
+	UNSYNC_ASSERT(CloneCommand.DestPostBytesOffset % AlignBytes == 0);
+
+	return CloneCommand;
+}
+
 // Returns extended absolute path of a form \\?\D:\verylongpath or \\?\UNC\servername\verylongpath
 // Expects an absolute path input. Returns original path on non-Windows.
 // https://docs.microsoft.com/en-us/windows/win32/fileio/maximum-file-path-limitation
@@ -868,7 +937,7 @@ FWindowsFile::CreateAsyncReader(uint32 MaxPipelineDepth)
 }
 
 bool
-CloneFileRange(FWindowsFile& SourceFile, FWindowsFile& DestFile, const FCopyCommand& CopyCommand, FBufferView FileData)
+CloneFileRange(FWindowsFile& SourceFile, FWindowsFile& DestFile, const FCopyCommand& CopyCommand, FBufferView FileData, uint64* ReflinkedBytesPtr)
 {
 	// ReFS has a number of preconditions for block cloning. See
 	// https://learn.microsoft.com/en-us/windows-server/storage/refs/block-cloning#functionality-restrictions-and-remarks
@@ -876,23 +945,12 @@ CloneFileRange(FWindowsFile& SourceFile, FWindowsFile& DestFile, const FCopyComm
 	// Block clones must be aligned on cluster boundaries, typically 4 kiB.
 	// TODO: Query the filesystem cluster size, since ReFS also supports 64 kiB clusters.
 	constexpr uint64 AlignBytes = 4096;
-	if (CopyCommand.SourceOffset % AlignBytes != CopyCommand.TargetOffset % AlignBytes)
-	{
-		return false;
-	}
-
-	const uint64 PreBytes	= (AlignBytes - (CopyCommand.SourceOffset % AlignBytes)) % AlignBytes;
-	const uint64 PostBytes	= (CopyCommand.SourceOffset + CopyCommand.Size) % AlignBytes;
-	// Check for writes that are too small to clone any blocks.
-	if (PreBytes >= CopyCommand.Size || PostBytes >= CopyCommand.Size || CopyCommand.Size < PreBytes + PostBytes + AlignBytes)
+	std::optional<FCloneCommand> OptCloneCommand = FCloneCommand::TryCreate(CopyCommand, AlignBytes);
+	if (!OptCloneCommand)
 	{
 		return false;
 	}
-	const uint64 CloneBytes = (CopyCommand.Size - PreBytes - PostBytes);
-	UNSYNC_ASSERT(CloneBytes % AlignBytes == 0);
-	UNSYNC_ASSERT(PreBytes + CloneBytes + PostBytes == CopyCommand.Size);
-	const uint64 SourceCloneOffset = CopyCommand.SourceOffset + PreBytes;
-	const uint64 DestCloneOffset   = CopyCommand.TargetOffset + PreBytes;
+	FCloneCommand& CloneCommand = *OptCloneCommand;
 
 	// Check whether the files are on the same filesystem, and the filesystem supports block cloning.
 	{
@@ -979,9 +1037,9 @@ CloneFileRange(FWindowsFile& SourceFile, FWindowsFile& DestFile, const FCopyComm
 
 	DUPLICATE_EXTENTS_DATA Extent{};
 	Extent.FileHandle				 = SourceFile.FileHandle;
-	Extent.SourceFileOffset.QuadPart = SourceCloneOffset;
-	Extent.TargetFileOffset.QuadPart = DestCloneOffset;
-	Extent.ByteCount.QuadPart		 = CloneBytes;
+	Extent.SourceFileOffset.QuadPart = CloneCommand.SourceCloneOffset;
+	Extent.TargetFileOffset.QuadPart = CloneCommand.DestCloneOffset;
+	Extent.ByteCount.QuadPart		 = CloneCommand.CloneBytes;
 	DWORD	BytesReturned{};
 	if (DeviceIoControl(DestFile.FileHandle, FSCTL_DUPLICATE_EXTENTS_TO_FILE, &Extent, sizeof(Extent), NULL, 0, &BytesReturned, NULL) != TRUE)
 	{
@@ -992,9 +1050,8 @@ CloneFileRange(FWindowsFile& SourceFile, FWindowsFile& DestFile, const FCopyComm
 	if (FileData.Size == CopyCommand.Size)
 	{
 		// Use file data that was already read into memory.
-		const uint64 PostOffset = PreBytes + CloneBytes;
-		if (DestFile.Write(FileData.Data, CopyCommand.TargetOffset, PreBytes) != PreBytes ||
-			DestFile.Write(FileData.Data + PostOffset, CopyCommand.TargetOffset + PostOffset, PostBytes) != PostBytes)
+		if (DestFile.Write(FileData.Data, CloneCommand.DestPreBytesOffset, CloneCommand.PreBytes) != CloneCommand.PreBytes ||
+			DestFile.Write(FileData.Data + CloneCommand.PreBytes + CloneCommand.CloneBytes, CloneCommand.DestPostBytesOffset, CloneCommand.PostBytes) != CloneCommand.PostBytes)
 		{
 			return false;
 		}
@@ -1003,24 +1060,29 @@ CloneFileRange(FWindowsFile& SourceFile, FWindowsFile& DestFile, const FCopyComm
 	{
 		// It's preferable for callers to pass in the already-in-memory data via `FileData` rather than re-reading it here.
 		uint8 Buffer[AlignBytes];
-		if (PreBytes > 0)
+		if (CloneCommand.PreBytes > 0)
 		{
-			if (SourceFile.Read(Buffer, CopyCommand.SourceOffset, PreBytes) != PreBytes ||
-				DestFile.Write(Buffer, CopyCommand.TargetOffset, PreBytes) != PreBytes)
+			if (SourceFile.Read(Buffer, CloneCommand.SourcePreBytesOffset, CloneCommand.PreBytes) != CloneCommand.PreBytes ||
+				DestFile.Write(Buffer, CloneCommand.DestPreBytesOffset, CloneCommand.PreBytes) != CloneCommand.PreBytes)
 			{
 				return false;
 			}
 		}
-		if (PostBytes > 0)
+		if (CloneCommand.PostBytes > 0)
 		{
-			if (SourceFile.Read(Buffer, CopyCommand.SourceOffset + PreBytes + CloneBytes, PostBytes) != PostBytes ||
-				DestFile.Write(Buffer, CopyCommand.TargetOffset + PreBytes + CloneBytes, PostBytes) != PostBytes)
+			if (SourceFile.Read(Buffer, CloneCommand.SourcePostBytesOffset, CloneCommand.PostBytes) != CloneCommand.PostBytes ||
+				DestFile.Write(Buffer, CloneCommand.DestPostBytesOffset, CloneCommand.PostBytes) != CloneCommand.PostBytes)
 			{
 				return false;
 			}
 		}
 	}
 
+	if (ReflinkedBytesPtr)
+	{
+		*ReflinkedBytesPtr += CloneCommand.CloneBytes;
+	}
+
 	return true;
 }
 
@@ -1085,25 +1147,33 @@ GetAvailableDiskSpace(const FPath& Path)
 #if UNSYNC_PLATFORM_UNIX
 FUnixFile::FUnixFile(const FPath& InFilename, EFileMode InMode, uint64 in_size) : Filename(InFilename), Mode(InMode)
 {
-	FileHandle = fopen(InFilename.native().c_str(), IsReadOnly(Mode) ? "rb" : "w+b");
-	if (FileHandle == nullptr)
+	FileDescriptor = open(
+		Filename.native().c_str(),
+		IsReadOnly(Mode) ? O_RDONLY : (O_RDWR | O_CREAT | O_TRUNC),
+		S_IRUSR | S_IWUSR | S_IRGRP | S_IWGRP | S_IROTH | S_IWOTH);
+
+	if (FileDescriptor == -1)
 	{
+		LastError = errno;
 		return;
 	}
 
-	FileDescriptor = fileno(FileHandle);
-
 	if (IsReadOnly(Mode))
 	{
 		struct stat stat_buf = {};
-		LastError			 = fstat(FileDescriptor, &stat_buf);
-		UNSYNC_ASSERT(LastError == 0);
-		FileSize = stat_buf.st_size;
+		if (fstat(FileDescriptor, &stat_buf) == 0)
+		{
+			FileSize = stat_buf.st_size;
+		}
+		else
+		{
+			LastError = errno;
+		}
 	}
 	else
 	{
-		LastError = ftruncate(FileDescriptor, in_size);
-		UNSYNC_ASSERT(LastError == 0);
+		// Note: Do not ftruncate() here as it inhibits copy-on-write reflinking.
+		// Delayed allocation will minimize fragmentation.
 		FileSize = in_size;
 	}
 }
@@ -1116,10 +1186,14 @@ FUnixFile::~FUnixFile()
 void
 FUnixFile::Close()
 {
-	if (FileHandle)
+	if (FileDescriptor != -1)
 	{
-		fclose(FileHandle);
-		FileHandle = nullptr;
+		if (close(FileDescriptor) != 0)
+		{
+			LastError = errno;
+			// Notify the user as this could mean file contents are not as expected. 
+			UNSYNC_WARNING(L"Failed to close %hs: %hs", Filename.string().c_str(), strerror(LastError));
+		}
 		FileDescriptor = -1;
 	}
 }
@@ -1151,8 +1225,54 @@ FUnixFile::Write(const void* data, uint64 DestOffset, uint64 WriteSize)
 }
 
 bool
-CloneFileRange(FUnixFile& SourceFile, FUnixFile& DestFile, const FCopyCommand& CopyCommand, FBufferView FileData)
-{
+CloneFileRange(FUnixFile& SourceFile, FUnixFile& DestFile, const FCopyCommand& CopyCommand, FBufferView FileData, uint64* NumReflinkedBytesPtr)
+{
+#if UNSYNC_PLATFORM_LINUX
+	if (const std::optional<FCloneCommand> OptCloneCommand = FCloneCommand::TryCreate(CopyCommand, 4096))
+	{
+		const FCloneCommand& CloneCommand = *OptCloneCommand;
+
+		// Most filesystems, including btrfs, will only reflink if the copy is block-aligned. This is also true of the FICLONERANGE ioctl.
+		// Unaligned copy_file_range() will succeed, but all file extents are duplicated, not reused.
+		// However by copying the aligned range separately, those blocks will be reflinked.
+
+		off_t SourcePreBytesOffset = CloneCommand.SourcePreBytesOffset;
+		off_t DestPreBytesOffset = CloneCommand.DestPreBytesOffset;
+		off_t SourceCloneOffset = CloneCommand.SourceCloneOffset;
+		off_t DestCloneOffset = CloneCommand.DestCloneOffset;
+		off_t SourcePostBytesOffset = CloneCommand.SourcePostBytesOffset;
+		off_t DestPostBytesOffset = CloneCommand.DestPostBytesOffset;
+
+		ssize_t PreBytesWritten = 0;
+		ssize_t CloneBytesWritten = 0;
+		ssize_t PostBytesWritten = 0;
+		if (
+			(PreBytesWritten = copy_file_range(
+				SourceFile.FileDescriptor, &SourcePreBytesOffset,
+				DestFile.FileDescriptor, &DestPreBytesOffset,
+				CloneCommand.PreBytes, 0)) == CloneCommand.PreBytes
+			&& 
+			(CloneBytesWritten = copy_file_range(
+				SourceFile.FileDescriptor, &SourceCloneOffset,
+				DestFile.FileDescriptor, &DestCloneOffset,
+				CloneCommand.CloneBytes, 0)) == CloneCommand.CloneBytes
+			&&
+			(PostBytesWritten = copy_file_range(
+				SourceFile.FileDescriptor, &SourcePostBytesOffset,
+				DestFile.FileDescriptor, &DestPostBytesOffset,
+				CloneCommand.PostBytes, 0)) == CloneCommand.PostBytes)
+		{
+			UNSYNC_ASSERT(PreBytesWritten + CloneBytesWritten + PostBytesWritten == CopyCommand.Size);
+
+			if (NumReflinkedBytesPtr)
+			{
+				*NumReflinkedBytesPtr += CloneBytesWritten;
+			}
+
+			return true;
+		}
+	}
+
 	off_t SourceOffset = CopyCommand.SourceOffset;
 	off_t DestOffset = CopyCommand.TargetOffset;
 
@@ -1171,6 +1291,9 @@ CloneFileRange(FUnixFile& SourceFile, FUnixFile& DestFile, const FCopyCommand& C
 	}
 
 	return Written == CopyCommand.Size;
+#else // !UNSYNC_PLATFORM_LINUX
+	return false;
+#endif // !UNSYNC_PLATFORM_LINUX
 }
 
 FFileAttributes
diff --git a/Engine/Source/Programs/Unsync/Private/UnsyncFile.h b/Engine/Source/Programs/Unsync/Private/UnsyncFile.h
--- a/Engine/Source/Programs/Unsync/Private/UnsyncFile.h
+++ b/Engine/Source/Programs/Unsync/Private/UnsyncFile.h
@@ -201,18 +201,14 @@ public:
 	virtual int64 WriteFrom(FIOReader& Reader, const FCopyCommand& CopyCommand, FBufferView FileData, uint64* NumBytesClonedPtr) override
 	{
 		FileType* FileReader = dynamic_cast<FileType*>(&Reader.ActualReader());
-		if (FileReader)
+		if (FileReader && FileReader->IsValid())
 		{
 			// Attempt block cloning
 			FileType* FileWriter = dynamic_cast<FileType*>(this);
 			UNSYNC_ASSERT(FileWriter);
 
-			if (CloneFileRange(*FileReader, *FileWriter, CopyCommand, FileData))
+			if (CloneFileRange(*FileReader, *FileWriter, CopyCommand, FileData, NumBytesClonedPtr))
 			{
-				if (NumBytesClonedPtr)
-				{
-					*NumBytesClonedPtr += CopyCommand.Size;
-				}
 				return CopyCommand.Size;
 			}
 		}
@@ -273,7 +269,7 @@ private:
 
 // Clone file blocks via ref counting if possible.
 // FileData optionally supplies an in-memory view of the file contents that's used for writing unaligned bytes at the beginning and end of the cloned range.
-bool CloneFileRange(FWindowsFile& SourceFile, FWindowsFile& DestFile, const FCopyCommand& CopyCommand, FBufferView FileData = FBufferView());
+bool CloneFileRange(FWindowsFile& SourceFile, FWindowsFile& DestFile, const FCopyCommand& CopyCommand, FBufferView FileData = FBufferView(), uint64* NumBytesClonedPtr = nullptr);
 
 using FNativeFile = FWindowsFile;
 #endif	// UNSYNC_PLATFORM_WINDOWS
@@ -283,7 +279,7 @@ struct FUnixFile;
 
 // Copy file data using in-kernel copy, which also provides block cloning (reflinking) on filesystems that support it.
 // The `FileData` parameter is not used; it's only for API compatiblity with the Windows CloneFileRange() function.
-bool CloneFileRange(FUnixFile& SourceFile, FUnixFile& DestFile, const FCopyCommand& CopyCommand, FBufferView FileData = FBufferView());
+bool CloneFileRange(FUnixFile& SourceFile, FUnixFile& DestFile, const FCopyCommand& CopyCommand, FBufferView FileData = FBufferView(), uint64* NumReflinkedBytesPtr = nullptr);
 
 struct FUnixFile : FFileReaderWriter<FUnixFile>
 {
@@ -292,7 +288,7 @@ struct FUnixFile : FFileReaderWriter<FUnixFile>
 
 	// IOBase
 	virtual uint64 GetSize() override { return FileSize; }
-	virtual bool   IsValid() override { return FileHandle != nullptr; }
+	virtual bool   IsValid() override { return FileDescriptor != -1 && LastError == 0; }
 	virtual void   Close() override;
 	virtual int32  GetError() override { return LastError; }
 
@@ -309,14 +305,10 @@ struct FUnixFile : FFileReaderWriter<FUnixFile>
 
 	static constexpr uint32 UNBUFFERED_READ_ALIGNMENT = 4096;
 
-	friend bool CloneFileRange(FUnixFile& SourceFile, FUnixFile& DestFile, const FCopyCommand& CopyCommand, FBufferView FileData);
-
-private:
-	bool OpenFileHandle(EFileMode InMode);
+	friend bool CloneFileRange(FUnixFile& SourceFile, FUnixFile& DestFile, const FCopyCommand& CopyCommand, FBufferView FileData, uint64* NumReflinkedBytesPtr);
 
 private:
 	EFileMode Mode			 = EFileMode::None;
-	FILE*	  FileHandle	 = nullptr;
 	int		  FileDescriptor = -1;
 };
 using FNativeFile = FUnixFile;
diff --git a/Engine/Source/Programs/Unsync/Private/UnsyncTarget.cpp b/Engine/Source/Programs/Unsync/Private/UnsyncTarget.cpp
--- a/Engine/Source/Programs/Unsync/Private/UnsyncTarget.cpp
+++ b/Engine/Source/Programs/Unsync/Private/UnsyncTarget.cpp
@@ -359,6 +359,13 @@ BuildTarget(FIOWriter& Output, FIOReader& Source, FIOReader& Base, const FNeedLi
 																												   uint64	  CmdReadSize,
 																												   uint64	  CmdUserData)
 			{
+				if (CmdReadSize == uint64(-1))
+				{
+					UNSYNC_ERROR(L"Read failed: %hs", FormatSystemErrorMessage(DataProvider.GetError()).c_str());
+					Error.Set(AppError(L"Failed to read blocks"));
+					return;
+				}
+
 				WriteTasks.run(
 					[Buffer = MakeShared(std::move(CmdBuffer)), CmdReadSize, Block, &DataProvider, &Output, &Error, &Stats, ListType]()
 					{
diff --git a/Engine/Source/Programs/Unsync/Private/UnsyncVersion.cpp b/Engine/Source/Programs/Unsync/Private/UnsyncVersion.cpp
--- a/Engine/Source/Programs/Unsync/Private/UnsyncVersion.cpp
+++ b/Engine/Source/Programs/Unsync/Private/UnsyncVersion.cpp
@@ -2,7 +2,7 @@
 
 #include "UnsyncVersion.h"
 
-#define UNSYNC_VERSION_STR "1.1.1"
+#define UNSYNC_VERSION_STR "1.1.2"
 
 namespace unsync {
 
diff --git a/Engine/Source/Programs/Unsync/Unsync.Build.cs b/Engine/Source/Programs/Unsync/Unsync.Build.cs
--- a/Engine/Source/Programs/Unsync/Unsync.Build.cs
+++ b/Engine/Source/Programs/Unsync/Unsync.Build.cs
@@ -19,6 +19,7 @@ public class Unsync : ModuleRules
 		if (Target.Platform == UnrealTargetPlatform.Win64)
 		{
 			PrivateDefinitions.Add("UNSYNC_PLATFORM_WINDOWS=1");
+			PrivateDefinitions.Add("UNSYNC_PLATFORM_LINUX=0");
 			PrivateDefinitions.Add("UNSYNC_PLATFORM_UNIX=0");
 			PrivateDefinitions.Add("NOMINMAX=1");
 			PrivateDefinitions.Add("WIN32_LEAN_AND_MEAN=1");
@@ -27,14 +28,17 @@ public class Unsync : ModuleRules
 			PrivateDefinitions.Add("_WINSOCK_DEPRECATED_NO_WARNINGS=1");
 			PrivateDefinitions.Add("_SILENCE_CXX17_CODECVT_HEADER_DEPRECATION_WARNING=1");
 		}
-		else if (Target.Platform == UnrealTargetPlatform.Mac || Target.Platform == UnrealTargetPlatform.Linux)
+		else if (Target.Platform == UnrealTargetPlatform.Mac)
 		{
-			PrivateDefinitions.Add("UNSYNC_PLATFORM_WINDOWS=0");
+			PrivateDefinitions.Add("UNSYNC_PLATFORM_LINUX=0");
 			PrivateDefinitions.Add("UNSYNC_PLATFORM_UNIX=1");
+			PrivateDefinitions.Add("UNSYNC_PLATFORM_WINDOWS=0");
 		}
-
-		if (Target.Platform == UnrealTargetPlatform.Linux)
+		else if (Target.Platform == UnrealTargetPlatform.Linux)
 		{
+			PrivateDefinitions.Add("UNSYNC_PLATFORM_LINUX=1");
+			PrivateDefinitions.Add("UNSYNC_PLATFORM_UNIX=1");
+			PrivateDefinitions.Add("UNSYNC_PLATFORM_WINDOWS=0");
 			PublicSystemLibraries.Add("dl");
 		}
 
