Unsync: Disable copy-on-write block cloning (reflinking) on Windows

#rnx
#Synced-CL 47002773

diff --git a/Engine/Source/Programs/Unsync/Private/UnsyncFile.cpp b/Engine/Source/Programs/Unsync/Private/UnsyncFile.cpp
--- a/Engine/Source/Programs/Unsync/Private/UnsyncFile.cpp
+++ b/Engine/Source/Programs/Unsync/Private/UnsyncFile.cpp
@@ -939,6 +939,10 @@ FWindowsFile::CreateAsyncReader(uint32 MaxPipelineDepth)
 bool
 CloneFileRange(FWindowsFile& SourceFile, FWindowsFile& DestFile, const FCopyCommand& CopyCommand, FBufferView FileData, uint64* ReflinkedBytesPtr)
 {
+#if true
+	// Disabled due to reports of slowness on NTFS. UE-354599
+	return false;
+#else // false
 	// ReFS has a number of preconditions for block cloning. See
 	// https://learn.microsoft.com/en-us/windows-server/storage/refs/block-cloning#functionality-restrictions-and-remarks
 
@@ -1084,6 +1088,7 @@ CloneFileRange(FWindowsFile& SourceFile, FWindowsFile& DestFile, const FCopyComm
 	}
 
 	return true;
+#endif // false
 }
 
 FFileAttributes
diff --git a/Engine/Source/Programs/Unsync/Private/UnsyncVersion.cpp b/Engine/Source/Programs/Unsync/Private/UnsyncVersion.cpp
--- a/Engine/Source/Programs/Unsync/Private/UnsyncVersion.cpp
+++ b/Engine/Source/Programs/Unsync/Private/UnsyncVersion.cpp
@@ -2,7 +2,7 @@
 
 #include "UnsyncVersion.h"
 
-#define UNSYNC_VERSION_STR "1.1.2"
+#define UNSYNC_VERSION_STR "1.1.3"
 
 namespace unsync {
 
