unsync - Add support for server announcements

* Parse announcements from /api/v1/hello response and display them once per session
* Support info, warning, and error announcement levels
* Add --no-server-announcements flag to suppress display
* Allow test preset as positional argument (e.g. `unsync test basic_hash`)
* Bump version to 1.1.4

#Synced-CL 50943321

diff --git a/Engine/Source/Programs/Unsync/Private/UnsyncCmdLogin.cpp b/Engine/Source/Programs/Unsync/Private/UnsyncCmdLogin.cpp
--- a/Engine/Source/Programs/Unsync/Private/UnsyncCmdLogin.cpp
+++ b/Engine/Source/Programs/Unsync/Private/UnsyncCmdLogin.cpp
@@ -45,6 +45,7 @@ CmdLogin(const FCmdLoginOptions& Options)
 	}
 
 	const FAuthDesc AuthDesc = FAuthDesc::FromHelloResponse(*HelloResponseResult);
+	ProxyQuery::DisplayAnnouncements(HelloResponseResult->Announcements);
 
 	TResult<FAuthToken> AuthTokenResult = Authenticate(AuthDesc, Options.Remote.RecvTimeoutSeconds, RefreshThreshold);
 
diff --git a/Engine/Source/Programs/Unsync/Private/UnsyncCmdQuery.cpp b/Engine/Source/Programs/Unsync/Private/UnsyncCmdQuery.cpp
--- a/Engine/Source/Programs/Unsync/Private/UnsyncCmdQuery.cpp
+++ b/Engine/Source/Programs/Unsync/Private/UnsyncCmdQuery.cpp
@@ -183,6 +183,7 @@ CmdQueryList(const FCmdQueryOptions& Options)
 		return -1;
 	}
 	FAuthDesc AuthDesc = FAuthDesc::FromHelloResponse(*HelloResponse);
+	ProxyQuery::DisplayAnnouncements(HelloResponse->Announcements);
 
 	if (Options.Args.empty())
 	{
@@ -240,6 +241,7 @@ CmdQuerySearch(const FCmdQueryOptions& Options)
 	}
 
 	FAuthDesc AuthDesc = FAuthDesc::FromHelloResponse(*HelloResponse);
+	ProxyQuery::DisplayAnnouncements(HelloResponse->Announcements);
 
 	const std::string& RootPath = Options.Args[0];
 	std::vector<std::regex> SubdirPatterns;
@@ -450,6 +452,7 @@ CmdQueryFile(const FCmdQueryOptions& Options)
 		return -1;
 	}
 	FAuthDesc AuthDesc = FAuthDesc::FromHelloResponse(*HelloResponse);
+	ProxyQuery::DisplayAnnouncements(HelloResponse->Announcements);
 	HelloConnection.Close();
 
 	TResult<FAuthToken> AuthToken = Authenticate(AuthDesc, Options.Remote.RecvTimeoutSeconds);
@@ -523,7 +526,8 @@ CmdQueryHttpGet(const FCmdQueryOptions& Options)
 			return -1;
 		}
 
-		FAuthDesc			AuthDesc  = FAuthDesc::FromHelloResponse(*HelloResponse);
+		FAuthDesc AuthDesc = FAuthDesc::FromHelloResponse(*HelloResponse);
+		ProxyQuery::DisplayAnnouncements(HelloResponse->Announcements);
 		TResult<FAuthToken> AuthToken = Authenticate(AuthDesc, Options.Remote.RecvTimeoutSeconds);
 
 		if (!AuthToken.IsOk())
diff --git a/Engine/Source/Programs/Unsync/Private/UnsyncHash.cpp b/Engine/Source/Programs/Unsync/Private/UnsyncHash.cpp
--- a/Engine/Source/Programs/Unsync/Private/UnsyncHash.cpp
+++ b/Engine/Source/Programs/Unsync/Private/UnsyncHash.cpp
@@ -232,34 +232,57 @@ HexCharToUint8(char C, uint8& Output)
 	return true;
 }
 
-bool
-ParseHashFromHexString(EStrongHashAlgorithmID StrongHasher, std::string_view HexString, FGenericHash& Output)
+// Decodes hex character pairs into bytes. HexString must have exactly NumBytes*2 characters.
+// When bReverseOrder is true, bytes are written in reverse order (for big-endian hex into little-endian memory).
+template<bool bReverseOrder = false>
+inline bool
+ParseBytesFromHexString(std::string_view HexString, uint8* OutputBytes, uint64 NumBytes)
 {
-	EHashType HashType = ToHashType(StrongHasher);
-
-	FGenericHash Result = {};
-	Result.Type			= HashType;
-
-	if (Result.Size() != HexString.length() / 2)
+	if (HexString.length() != NumBytes * 2)
 	{
 		return false;
 	}
 
-	const uint64 NumOutputBytes = Result.Size();
-
-	for (uint64 Index = 0; Index < NumOutputBytes; ++Index)
+	for (uint64 Index = 0; Index < NumBytes; ++Index)
 	{
-		uint8 L = 0, H = 0;
+		uint8 H = 0, L = 0;
 		bool  bValid = HexCharToUint8(HexString[Index * 2 + 0], H) && HexCharToUint8(HexString[Index * 2 + 1], L);
 		if (!bValid)
 		{
 			return false;
 		}
-		Result.Data[Index] = (H << 4) | L;
+		uint64 OutputIndex = bReverseOrder ? (NumBytes - 1 - Index) : Index;
+		OutputBytes[OutputIndex] = (H << 4) | L;
+	}
+
+	return true;
+}
+
+bool
+ParseHashFromHexString(EStrongHashAlgorithmID StrongHasher, std::string_view HexString, FGenericHash& Output)
+{
+	FGenericHash Result = {};
+	Result.Type			= ToHashType(StrongHasher);
+
+	if (!ParseBytesFromHexString(HexString, Result.Data, Result.Size()))
+	{
+		return false;
 	}
 
 	Output = Result;
+	return true;
+}
 
+bool
+ParseUint64FromHexString(std::string_view HexString, uint64& Output)
+{
+	uint64 Result = 0;
+	if (!ParseBytesFromHexString</*bReverseOrder*/ true>(HexString, reinterpret_cast<uint8*>(&Result), 8))
+	{
+		return false;
+	}
+
+	Output = Result;
 	return true;
 }
 
diff --git a/Engine/Source/Programs/Unsync/Private/UnsyncHash.h b/Engine/Source/Programs/Unsync/Private/UnsyncHash.h
--- a/Engine/Source/Programs/Unsync/Private/UnsyncHash.h
+++ b/Engine/Source/Programs/Unsync/Private/UnsyncHash.h
@@ -151,6 +151,7 @@ struct FGenericHash
 };
 
 bool ParseHashFromHexString(EStrongHashAlgorithmID StrongHasher, std::string_view HexString, FGenericHash& Output);
+bool ParseUint64FromHexString(std::string_view HexString, uint64& Output);
 
 inline FHash160
 ToHash160(const FHash256& Other)
diff --git a/Engine/Source/Programs/Unsync/Private/UnsyncMain.cpp b/Engine/Source/Programs/Unsync/Private/UnsyncMain.cpp
--- a/Engine/Source/Programs/Unsync/Private/UnsyncMain.cpp
+++ b/Engine/Source/Programs/Unsync/Private/UnsyncMain.cpp
@@ -151,6 +151,7 @@ InnerMain(int Argc, char** Argv)
 		App->add_option("--protocol", ProtocolName, "Explicitly specify server protocol instead of inferring it from URL")
 			->required(false)
 			->check(CLI::IsMember({"unsync", "jupiter", "horde"}));
+		App->add_flag("--no-server-announcements", GNoServerAnnouncements, "Suppress server announcement messages");
 	};
 
 	// Configure hash
@@ -367,7 +368,7 @@ InnerMain(int Argc, char** Argv)
 	SubCommands.push_back(SubPatch);
 
 	CLI::App* SubTest = Cli.add_subcommand("test", "Run internal tests");
-	SubTest->add_option("--preset", PresetUtf8, "Test preset")->default_str(PresetUtf8);
+	SubTest->add_option("Preset, --preset", PresetUtf8, "Test preset")->default_str(PresetUtf8);
 	SubCommands.push_back(SubTest);
 
 	// Configure query
diff --git a/Engine/Source/Programs/Unsync/Private/UnsyncProxy.cpp b/Engine/Source/Programs/Unsync/Private/UnsyncProxy.cpp
--- a/Engine/Source/Programs/Unsync/Private/UnsyncProxy.cpp
+++ b/Engine/Source/Programs/Unsync/Private/UnsyncProxy.cpp
@@ -4,6 +4,7 @@
 #include "UnsyncAuth.h"
 #include "UnsyncCompression.h"
 #include "UnsyncFile.h"
+#include "UnsyncHash.h"
 #include "UnsyncJupiter.h"
 #include "UnsyncPool.h"
 #include "UnsyncProgress.h"
@@ -816,10 +817,38 @@ FUnsyncProtocolImpl::QueryHello(FHttpConnection& HttpConnection, const FAuthDesc
 				{
 					Result.Features.bBlockDownload = true;
 				}
+				else if (Elem.string_value() == "announcements")
+				{
+					Result.Features.bAnnouncements = true;
+				}
 			}
 		}
 	}
 
+	if (auto& Field = JsonObject["announcements"]; Field.is_array())
+	{
+		for (const auto& Elem : Field.array_items())
+		{
+			if (!Elem.is_object())
+			{
+				continue;
+			}
+
+			ProxyQuery::FAnnouncement Announcement;
+
+			const auto& IdField = Elem["id"];
+			if (!IdField.is_string() || !ParseUint64FromHexString(IdField.string_value(), Announcement.Id))
+			{
+				continue;
+			}
+
+			Announcement.Level	 = Elem["level"].string_value();
+			Announcement.Message = Elem["message"].string_value();
+
+			Result.Announcements.push_back(std::move(Announcement));
+		}
+	}
+
 	if (auto& Field = JsonObject["primary"]; Field.is_string())
 	{
 		const std::string&	 PrimaryHostStr	 = Field.string_value();
@@ -835,6 +864,40 @@ FUnsyncProtocolImpl::QueryHello(FHttpConnection& HttpConnection, const FAuthDesc
 	return ResultOk(std::move(Result));
 }
 
+bool GNoServerAnnouncements = false;
+
+void
+ProxyQuery::DisplayAnnouncements(const std::vector<ProxyQuery::FAnnouncement>& Announcements)
+{
+	if (GNoServerAnnouncements || Announcements.empty())
+	{
+		return;
+	}
+
+	static std::unordered_set<uint64> DisplayedIds;
+
+	for (const FAnnouncement& Announcement : Announcements)
+	{
+		if (!DisplayedIds.insert(Announcement.Id).second)
+		{
+			continue;
+		}
+
+		if (Announcement.Level == "error")
+		{
+			UNSYNC_ERROR(L"%hs", Announcement.Message.c_str());
+		}
+		else if (Announcement.Level == "warning")
+		{
+			UNSYNC_WARNING(L"%hs", Announcement.Message.c_str());
+		}
+		else
+		{
+			UNSYNC_LOG(L"%hs", Announcement.Message.c_str());
+		}
+	}
+}
+
 TResult<ProxyQuery::FDirectoryListing>
 ProxyQuery::FDirectoryListing::FromJson(const char* JsonString)
 {
@@ -1368,6 +1431,8 @@ FProxyPool::FProxyPool(const FRemoteDesc& InRemoteDesc, const FAuthDesc* InAuthD
 
 			Features  = Data.Features;
 			SessionId = Data.SessionId;
+
+			ProxyQuery::DisplayAnnouncements(Data.Announcements);
 		}
 
 		bValid = Response.IsOk();
diff --git a/Engine/Source/Programs/Unsync/Private/UnsyncProxy.h b/Engine/Source/Programs/Unsync/Private/UnsyncProxy.h
--- a/Engine/Source/Programs/Unsync/Private/UnsyncProxy.h
+++ b/Engine/Source/Programs/Unsync/Private/UnsyncProxy.h
@@ -109,6 +109,7 @@ struct FRemoteProtocolFeatures
 	bool bFileDownload	   = false;
 	bool bManifestDownload = false;
 	bool bBlockDownload	   = false;
+	bool bAnnouncements	   = false;
 };
 
 struct FTelemetryEventSyncComplete
@@ -227,8 +228,17 @@ struct FPooledHttpConnection
 	std::unique_ptr<FHttpConnection> Inner;
 };
 
+extern bool GNoServerAnnouncements;
+
 namespace ProxyQuery {
 
+struct FAnnouncement
+{
+	uint64		Id = 0;	 // 64-bit hash, parsed from hex string in JSON
+	std::string Level;	 // "info", "warning", "error"
+	std::string Message;
+};
+
 struct FHelloResponse
 {
 	std::string Name;
@@ -246,6 +256,8 @@ struct FHelloResponse
 
 	std::optional<FHostAddressAndPort> PrimaryHost;
 
+	std::vector<FAnnouncement> Announcements;
+
 	// Derived data
 
 	bool bConnectionEncrypted = false;
@@ -255,6 +267,8 @@ struct FHelloResponse
 TResult<FHelloResponse> Hello(const FRemoteDesc& RemoteDesc, const FAuthDesc* OptAuthDesc = nullptr);
 TResult<FHelloResponse> Hello(EProtocolFlavor Protocol, FHttpConnection& Connection, const FAuthDesc* OptAuthDesc = nullptr);
 
+void DisplayAnnouncements(const std::vector<FAnnouncement>& Announcements);
+
 struct FDirectoryListingEntry
 {
 	std::string Name;  // utf-8
diff --git a/Engine/Source/Programs/Unsync/Private/UnsyncTest.cpp b/Engine/Source/Programs/Unsync/Private/UnsyncTest.cpp
--- a/Engine/Source/Programs/Unsync/Private/UnsyncTest.cpp
+++ b/Engine/Source/Programs/Unsync/Private/UnsyncTest.cpp
@@ -833,6 +833,48 @@ TestBasicHash()
 		UNSYNC_ASSERT(HashMap[K2] == 2);
 		UNSYNC_ASSERT(HashMap[K3] == 3);
 	}
+
+	// ParseHashFromHexString
+	{
+		FGenericHash Parsed;
+		UNSYNC_ASSERT(ParseHashFromHexString(EStrongHashAlgorithmID::Blake3_128, "2c4c1fa09b1a3459bc56ac6af6b446c8", Parsed));
+		std::string RoundTrip = BytesToHexString(Parsed.Data, Parsed.Size());
+		UNSYNC_ASSERT(RoundTrip == "2c4c1fa09b1a3459bc56ac6af6b446c8");
+
+		UNSYNC_ASSERT(ParseHashFromHexString(EStrongHashAlgorithmID::Blake3_160, "2c4c1fa09b1a3459bc56ac6af6b446c89c784cf9", Parsed));
+		RoundTrip = BytesToHexString(Parsed.Data, Parsed.Size());
+		UNSYNC_ASSERT(RoundTrip == "2c4c1fa09b1a3459bc56ac6af6b446c89c784cf9");
+
+		// Wrong length
+		UNSYNC_ASSERT(!ParseHashFromHexString(EStrongHashAlgorithmID::Blake3_128, "2c4c1fa0", Parsed));
+		// Invalid hex char
+		UNSYNC_ASSERT(!ParseHashFromHexString(EStrongHashAlgorithmID::Blake3_128, "zz4c1fa09b1a3459bc56ac6af6b446c8", Parsed));
+	}
+
+	// ParseUint64FromHexString
+	{
+		uint64 Value = 0;
+		UNSYNC_ASSERT(ParseUint64FromHexString("0000000000000000", Value));
+		UNSYNC_ASSERT(Value == 0);
+
+		UNSYNC_ASSERT(ParseUint64FromHexString("0000000000000001", Value));
+		UNSYNC_ASSERT(Value == 1);
+
+		UNSYNC_ASSERT(ParseUint64FromHexString("ffffffffffffffff", Value));
+		UNSYNC_ASSERT(Value == UINT64_MAX);
+
+		UNSYNC_ASSERT(ParseUint64FromHexString("1584c5a558e0bc7f", Value));
+		UNSYNC_ASSERT(Value == 0x1584c5a558e0bc7f);
+
+		UNSYNC_ASSERT(ParseUint64FromHexString("ABCDEF0123456789", Value));
+		UNSYNC_ASSERT(Value == 0xABCDEF0123456789);
+
+		// Wrong length
+		UNSYNC_ASSERT(!ParseUint64FromHexString("1234", Value));
+		UNSYNC_ASSERT(!ParseUint64FromHexString("00000000000000001", Value));
+		// Invalid hex char
+		UNSYNC_ASSERT(!ParseUint64FromHexString("zzzzzzzzzzzzzzzz", Value));
+	}
 }
 
 void
diff --git a/Engine/Source/Programs/Unsync/Private/UnsyncVersion.cpp b/Engine/Source/Programs/Unsync/Private/UnsyncVersion.cpp
--- a/Engine/Source/Programs/Unsync/Private/UnsyncVersion.cpp
+++ b/Engine/Source/Programs/Unsync/Private/UnsyncVersion.cpp
@@ -2,7 +2,7 @@
 
 #include "UnsyncVersion.h"
 
-#define UNSYNC_VERSION_STR "1.1.3"
+#define UNSYNC_VERSION_STR "1.1.4"
 
 namespace unsync {
 
