unsync - Add an option to only use stdout (not stderr) and polish server announcement logging

* Use --stdout to turn off stderr output
* Bump version to 1.1.5

#Synced-CL 50943321

diff --git a/Engine/Source/Programs/Unsync/Private/UnsyncCmdLogin.cpp b/Engine/Source/Programs/Unsync/Private/UnsyncCmdLogin.cpp
--- a/Engine/Source/Programs/Unsync/Private/UnsyncCmdLogin.cpp
+++ b/Engine/Source/Programs/Unsync/Private/UnsyncCmdLogin.cpp
@@ -45,7 +45,6 @@ CmdLogin(const FCmdLoginOptions& Options)
 	}
 
 	const FAuthDesc AuthDesc = FAuthDesc::FromHelloResponse(*HelloResponseResult);
-	ProxyQuery::DisplayAnnouncements(HelloResponseResult->Announcements);
 
 	TResult<FAuthToken> AuthTokenResult = Authenticate(AuthDesc, Options.Remote.RecvTimeoutSeconds, RefreshThreshold);
 
@@ -61,7 +60,7 @@ CmdLogin(const FCmdLoginOptions& Options)
 			{
 				Request.Url = "/api/v1/login";
 			}
-			if (Options.Remote.Protocol == EProtocolFlavor::Horde)
+			else if (Options.Remote.Protocol == EProtocolFlavor::Horde)
 			{
 				Request.Url = "/api/v1/projects";
 			}
diff --git a/Engine/Source/Programs/Unsync/Private/UnsyncCmdQuery.cpp b/Engine/Source/Programs/Unsync/Private/UnsyncCmdQuery.cpp
--- a/Engine/Source/Programs/Unsync/Private/UnsyncCmdQuery.cpp
+++ b/Engine/Source/Programs/Unsync/Private/UnsyncCmdQuery.cpp
@@ -183,7 +183,7 @@ CmdQueryList(const FCmdQueryOptions& Options)
 		return -1;
 	}
 	FAuthDesc AuthDesc = FAuthDesc::FromHelloResponse(*HelloResponse);
-	ProxyQuery::DisplayAnnouncements(HelloResponse->Announcements);
+
 
 	if (Options.Args.empty())
 	{
@@ -241,7 +241,7 @@ CmdQuerySearch(const FCmdQueryOptions& Options)
 	}
 
 	FAuthDesc AuthDesc = FAuthDesc::FromHelloResponse(*HelloResponse);
-	ProxyQuery::DisplayAnnouncements(HelloResponse->Announcements);
+
 
 	const std::string& RootPath = Options.Args[0];
 	std::vector<std::regex> SubdirPatterns;
@@ -452,7 +452,7 @@ CmdQueryFile(const FCmdQueryOptions& Options)
 		return -1;
 	}
 	FAuthDesc AuthDesc = FAuthDesc::FromHelloResponse(*HelloResponse);
-	ProxyQuery::DisplayAnnouncements(HelloResponse->Announcements);
+
 	HelloConnection.Close();
 
 	TResult<FAuthToken> AuthToken = Authenticate(AuthDesc, Options.Remote.RecvTimeoutSeconds);
@@ -527,7 +527,7 @@ CmdQueryHttpGet(const FCmdQueryOptions& Options)
 		}
 
 		FAuthDesc AuthDesc = FAuthDesc::FromHelloResponse(*HelloResponse);
-		ProxyQuery::DisplayAnnouncements(HelloResponse->Announcements);
+
 		TResult<FAuthToken> AuthToken = Authenticate(AuthDesc, Options.Remote.RecvTimeoutSeconds);
 
 		if (!AuthToken.IsOk())
diff --git a/Engine/Source/Programs/Unsync/Private/UnsyncLog.cpp b/Engine/Source/Programs/Unsync/Private/UnsyncLog.cpp
--- a/Engine/Source/Programs/Unsync/Private/UnsyncLog.cpp
+++ b/Engine/Source/Programs/Unsync/Private/UnsyncLog.cpp
@@ -55,6 +55,7 @@ bool				GBreakOnWarning		= false;
 thread_local uint32 GLogIndent			= 0;
 bool				GLogProgress		= false;
 bool				GLogMachineReadable = false;
+bool				GLogStdOut		= false;
 
 std::mutex GLogMutex;
 
@@ -78,6 +79,12 @@ LogSaveCommandLineUtf8(int Argc, char** Argv)
 static FILE*
 GetLogStream(ELogLevel LogLevel)
 {
+	// When --stdout is set, all output goes to stdout (no stderr).
+	if (GLogStdOut)
+	{
+		return stdout;
+	}
+
 	// If machine readable log format is enabled, *all* diagnostic logging is sent to stderr.
 	// Otherwise, only errors and warnings are sent to stderr, while normal diagnostics are sent to stdout.
 
diff --git a/Engine/Source/Programs/Unsync/Private/UnsyncLog.h b/Engine/Source/Programs/Unsync/Private/UnsyncLog.h
--- a/Engine/Source/Programs/Unsync/Private/UnsyncLog.h
+++ b/Engine/Source/Programs/Unsync/Private/UnsyncLog.h
@@ -14,6 +14,7 @@ extern bool				   GBreakOnWarning;
 extern thread_local uint32 GLogIndent;
 extern bool				   GLogProgress;		 // Whether to output @progress and @status markers to stdout
 extern bool				   GLogMachineReadable;	 // Whether output is intended for other programs rather than humans
+extern bool				   GLogStdOut;           // Whether to send all output to stdout (disabling stderr)
 
 struct FError;
 
diff --git a/Engine/Source/Programs/Unsync/Private/UnsyncMain.cpp b/Engine/Source/Programs/Unsync/Private/UnsyncMain.cpp
--- a/Engine/Source/Programs/Unsync/Private/UnsyncMain.cpp
+++ b/Engine/Source/Programs/Unsync/Private/UnsyncMain.cpp
@@ -413,6 +413,7 @@ InnerMain(int Argc, char** Argv)
 		auto VeryVerboseFlag = Subcommand->add_flag("--very-verbose", GLogVeryVerbose, "Very verbose logging");
 		auto SilentFlag		 = Subcommand->add_flag("--silent", GLogSilent, "Skip all console logging except errors and warnings");
 		Subcommand->add_flag("--progress", GLogProgress, "Output @progress and @status markers");
+		Subcommand->add_flag("--stdout", GLogStdOut, "Send all output to stdout (no stderr)");
 		Subcommand->add_option("--threads", GMaxThreads, "Limit worker threads to specified number");
 		Subcommand->add_flag("--buffered-files", GForceBufferedFiles, "Always use buffered file IO");
 		Subcommand->add_flag("--debug", bUseDebugMode, "Enable extra debugging features, such as extra memory safety validation");
diff --git a/Engine/Source/Programs/Unsync/Private/UnsyncProxy.cpp b/Engine/Source/Programs/Unsync/Private/UnsyncProxy.cpp
--- a/Engine/Source/Programs/Unsync/Private/UnsyncProxy.cpp
+++ b/Engine/Source/Programs/Unsync/Private/UnsyncProxy.cpp
@@ -680,29 +680,39 @@ FUnsyncProtocolImpl::GetSocketSecurity() const
 }
 
 TResult<ProxyQuery::FHelloResponse>
-ProxyQuery::Hello(const FRemoteDesc& RemoteDesc, const FAuthDesc* OptAuthDesc)
+ProxyQuery::Hello(const FRemoteDesc& RemoteDesc, const FAuthDesc* OptAuthDesc, bool bDisplayAnnouncements)
 {
 	FTlsClientSettings TlsSettings = RemoteDesc.GetTlsClientSettings();
 	FHttpConnection	   Connection(RemoteDesc.Host.Address, RemoteDesc.Host.Port, RemoteDesc.TlsRequirement, TlsSettings, RemoteDesc.RecvTimeoutSeconds);
 
-	return Hello(RemoteDesc.Protocol, Connection, OptAuthDesc);
+	return Hello(RemoteDesc.Protocol, Connection, OptAuthDesc, bDisplayAnnouncements);
 }
 
 TResult<ProxyQuery::FHelloResponse>
-ProxyQuery::Hello(EProtocolFlavor Protocol, FHttpConnection& Connection, const FAuthDesc* OptAuthDesc)
+ProxyQuery::Hello(EProtocolFlavor Protocol, FHttpConnection& Connection, const FAuthDesc* OptAuthDesc, bool bDisplayAnnouncements)
 {
-	if (Protocol == EProtocolFlavor::Horde)
+	TResult<FHelloResponse> Result = [&]() -> TResult<FHelloResponse>
 	{
-		return FHordeProtocolImpl::QueryHello(Connection);
-	}
-	else if (Protocol == EProtocolFlavor::Unsync)
-	{
-		return FUnsyncProtocolImpl::QueryHello(Connection, OptAuthDesc);
-	}
-	else
+		if (Protocol == EProtocolFlavor::Horde)
+		{
+			return FHordeProtocolImpl::QueryHello(Connection);
+		}
+		else if (Protocol == EProtocolFlavor::Unsync)
+		{
+			return FUnsyncProtocolImpl::QueryHello(Connection, OptAuthDesc);
+		}
+		else
+		{
+			return AppError("Protocol does not support server information query");
+		}
+	}();
+
+	if (Result.IsOk() && bDisplayAnnouncements)
 	{
-		return AppError("Protocol does not support server information query");
+		DisplayAnnouncements(Result->Announcements);
 	}
+
+	return Result;
 }
 
 TResult<ProxyQuery::FHelloResponse>
@@ -883,17 +893,33 @@ ProxyQuery::DisplayAnnouncements(const std::vector<ProxyQuery::FAnnouncement>& A
 			continue;
 		}
 
-		if (Announcement.Level == "error")
-		{
-			UNSYNC_ERROR(L"%hs", Announcement.Message.c_str());
-		}
-		else if (Announcement.Level == "warning")
+		// Split on newlines so each line gets the appropriate log prefix
+		const std::string& Msg = Announcement.Message;
+		uint64 Pos = 0;
+		while (Pos <= Msg.size())
 		{
-			UNSYNC_WARNING(L"%hs", Announcement.Message.c_str());
-		}
-		else
-		{
-			UNSYNC_LOG(L"%hs", Announcement.Message.c_str());
+			uint64 End = Msg.find('\n', Pos);
+			if (End == std::string::npos)
+			{
+				End = Msg.size();
+			}
+
+			std::string Line = Msg.substr(Pos, End - Pos);
+
+			if (Announcement.Level == "error")
+			{
+				UNSYNC_ERROR(L"%hs", Line.c_str());
+			}
+			else if (Announcement.Level == "warning")
+			{
+				UNSYNC_WARNING(L"%hs", Line.c_str());
+			}
+			else
+			{
+				UNSYNC_LOG(L"%hs", Line.c_str());
+			}
+
+			Pos = End + 1;
 		}
 	}
 }
@@ -1431,8 +1457,6 @@ FProxyPool::FProxyPool(const FRemoteDesc& InRemoteDesc, const FAuthDesc* InAuthD
 
 			Features  = Data.Features;
 			SessionId = Data.SessionId;
-
-			ProxyQuery::DisplayAnnouncements(Data.Announcements);
 		}
 
 		bValid = Response.IsOk();
diff --git a/Engine/Source/Programs/Unsync/Private/UnsyncProxy.h b/Engine/Source/Programs/Unsync/Private/UnsyncProxy.h
--- a/Engine/Source/Programs/Unsync/Private/UnsyncProxy.h
+++ b/Engine/Source/Programs/Unsync/Private/UnsyncProxy.h
@@ -264,8 +264,8 @@ struct FHelloResponse
 
 	bool SupportsAuthentication() const { return Features.bAuthentication && !AuthServerUri.empty() && !AuthClientId.empty(); }
 };
-TResult<FHelloResponse> Hello(const FRemoteDesc& RemoteDesc, const FAuthDesc* OptAuthDesc = nullptr);
-TResult<FHelloResponse> Hello(EProtocolFlavor Protocol, FHttpConnection& Connection, const FAuthDesc* OptAuthDesc = nullptr);
+TResult<FHelloResponse> Hello(const FRemoteDesc& RemoteDesc, const FAuthDesc* OptAuthDesc = nullptr, bool bDisplayAnnouncements = true);
+TResult<FHelloResponse> Hello(EProtocolFlavor Protocol, FHttpConnection& Connection, const FAuthDesc* OptAuthDesc = nullptr, bool bDisplayAnnouncements = true);
 
 void DisplayAnnouncements(const std::vector<FAnnouncement>& Announcements);
 
diff --git a/Engine/Source/Programs/Unsync/Private/UnsyncUtil.cpp b/Engine/Source/Programs/Unsync/Private/UnsyncUtil.cpp
--- a/Engine/Source/Programs/Unsync/Private/UnsyncUtil.cpp
+++ b/Engine/Source/Programs/Unsync/Private/UnsyncUtil.cpp
@@ -447,14 +447,14 @@ StringEscape(const std::string_view Input)
 			snprintf(buf, sizeof buf, "\\u%04x", C);
 			Result += buf;
 		}
-		else if (static_cast<uint8_t>(C) == 0xe2 && static_cast<uint8_t>(Input[i + 1]) == 0x80 &&
-				 static_cast<uint8_t>(Input[i + 2]) == 0xa8)
+		else if (static_cast<uint8_t>(C) == 0xe2 && i + 2 < Input.length() &&
+				 static_cast<uint8_t>(Input[i + 1]) == 0x80 && static_cast<uint8_t>(Input[i + 2]) == 0xa8)
 		{
 			Result += "\\u2028";
 			i += 2;
 		}
-		else if (static_cast<uint8_t>(C) == 0xe2 && static_cast<uint8_t>(Input[i + 1]) == 0x80 &&
-				 static_cast<uint8_t>(Input[i + 2]) == 0xa9)
+		else if (static_cast<uint8_t>(C) == 0xe2 && i + 2 < Input.length() &&
+				 static_cast<uint8_t>(Input[i + 1]) == 0x80 && static_cast<uint8_t>(Input[i + 2]) == 0xa9)
 		{
 			Result += "\\u2029";
 			i += 2;
diff --git a/Engine/Source/Programs/Unsync/Private/UnsyncVersion.cpp b/Engine/Source/Programs/Unsync/Private/UnsyncVersion.cpp
--- a/Engine/Source/Programs/Unsync/Private/UnsyncVersion.cpp
+++ b/Engine/Source/Programs/Unsync/Private/UnsyncVersion.cpp
@@ -2,7 +2,7 @@
 
 #include "UnsyncVersion.h"
 
-#define UNSYNC_VERSION_STR "1.1.4"
+#define UNSYNC_VERSION_STR "1.1.5"
 
 namespace unsync {
 
